Fixed up redundant FText .ToString() usage and some improper FText usage.

diff --git a/Engine/Source/Developer/AutomationWindow/Private/SAutomationWindow.cpp b/Engine/Source/Developer/AutomationWindow/Private/SAutomationWindow.cpp
--- a/Engine/Source/Developer/AutomationWindow/Private/SAutomationWindow.cpp
+++ b/Engine/Source/Developer/AutomationWindow/Private/SAutomationWindow.cpp
@@ -606,17 +606,16 @@ END_SLATE_FUNCTION_BUILD_OPTIMIZATION
 
 FText SAutomationWindow::CreateDeviceTooltip(int32 ClusterIndex)
 {
-	//the tool tip
-	FString ToolTip = LOCTEXT("ToolTipGameInstances", "Game Instances:").ToString();
+	FTextBuilder ReportBuilder;
+	ReportBuilder.AppendLine(LOCTEXT("ToolTipGameInstances", "Game Instances:"));
 
 	int32 NumDevices = AutomationController->GetNumDevicesInCluster( ClusterIndex );
 	for ( int32 DeviceIndex = 0; DeviceIndex < NumDevices; ++DeviceIndex )
 	{
-		ToolTip += TEXT("\n   ");
-		ToolTip += AutomationController->GetGameInstanceName( ClusterIndex, DeviceIndex );
+		ReportBuilder.AppendLine(AutomationController->GetGameInstanceName(ClusterIndex, DeviceIndex));
 	}
 
-	return FText::FromString( ToolTip );
+	return ReportBuilder.ToText();
 }
 
 
diff --git a/Engine/Source/Developer/AutomationWindow/Private/SAutomationWindowCommandBar.cpp b/Engine/Source/Developer/AutomationWindow/Private/SAutomationWindowCommandBar.cpp
--- a/Engine/Source/Developer/AutomationWindow/Private/SAutomationWindowCommandBar.cpp
+++ b/Engine/Source/Developer/AutomationWindow/Private/SAutomationWindowCommandBar.cpp
@@ -38,8 +38,8 @@ void SAutomationWindowCommandBar::Construct( const FArguments& InArgs, const TSh
 				SAssignNew(CopyButton, SButton)
 					.ContentPadding(FMargin(6.0, 2.0))
 					.IsEnabled(false)
-					.Text(LOCTEXT("AutomationCopyButtonText", "Copy").ToString())
-					.ToolTipText(LOCTEXT("AutomationCopyButtonTooltip", "Copy the selected log messages to the clipboard").ToString())
+					.Text(LOCTEXT("AutomationCopyButtonText", "Copy"))
+					.ToolTipText(LOCTEXT("AutomationCopyButtonTooltip", "Copy the selected log messages to the clipboard"))
 					.OnClicked(this, &SAutomationWindowCommandBar::HandleCopyButtonClicked)
 			]
 	];
diff --git a/Engine/Source/Developer/CollisionAnalyzer/Private/SCAQueryDetails.cpp b/Engine/Source/Developer/CollisionAnalyzer/Private/SCAQueryDetails.cpp
--- a/Engine/Source/Developer/CollisionAnalyzer/Private/SCAQueryDetails.cpp
+++ b/Engine/Source/Developer/CollisionAnalyzer/Private/SCAQueryDetails.cpp
@@ -161,25 +161,25 @@ void SCAQueryDetails::Construct(const FArguments& InArgs, TSharedPtr<SCollisionA
 					.Padding(2)
 					[
 						SNew(STextBlock)
-						.Text(LOCTEXT("QueryStart", "Start:").ToString())
+						.Text(LOCTEXT("QueryStart", "Start:"))
 					]
 					+SGridPanel::Slot(1,0)
 					.Padding(2)
 					[
 						SNew(STextBlock)
-						.Text(this, &SCAQueryDetails::GetStartString)
+						.Text(this, &SCAQueryDetails::GetStartText)
 					]
 					+SGridPanel::Slot(0,1)
 					.Padding(2)
 					[
 						SNew(STextBlock)
-						.Text(LOCTEXT("QueryEnd", "End:").ToString())
+						.Text(LOCTEXT("QueryEnd", "End:"))
 					]
 					+SGridPanel::Slot(1,1)
 					.Padding(2)
 					[
 						SNew(STextBlock)
-						.Text(this, &SCAQueryDetails::GetEndString)
+						.Text(this, &SCAQueryDetails::GetEndText)
 					]
 				]
 				// Right has controls
@@ -193,7 +193,7 @@ void SCAQueryDetails::Construct(const FArguments& InArgs, TSharedPtr<SCollisionA
 					.Content()
 					[
 						SNew(STextBlock)
-						.Text(LOCTEXT("ShowMisses", "Show Misses").ToString())
+						.Text(LOCTEXT("ShowMisses", "Show Misses"))
 					]
 				]
 			]
@@ -214,10 +214,10 @@ void SCAQueryDetails::Construct(const FArguments& InArgs, TSharedPtr<SCollisionA
 				.OnGenerateRow(this, &SCAQueryDetails::ResultListGenerateRow)
 				.HeaderRow(
 					SNew(SHeaderRow)
-					+SHeaderRow::Column("Time").DefaultLabel(LOCTEXT("ResultListTimeHeader", "Time").ToString()).FillWidth(0.7)
-					+SHeaderRow::Column("Type").DefaultLabel(LOCTEXT("ResultListTypeHeader", "Type").ToString()).FillWidth(0.7)
-					+SHeaderRow::Column("Component").DefaultLabel(LOCTEXT("ResultListComponentHeader", "Component").ToString()).FillWidth(3)
-					+SHeaderRow::Column("Normal").DefaultLabel(LOCTEXT("ResultListNormalHeader", "Normal").ToString()).FillWidth(1.8)
+					+SHeaderRow::Column("Time").DefaultLabel(LOCTEXT("ResultListTimeHeader", "Time")).FillWidth(0.7)
+					+SHeaderRow::Column("Type").DefaultLabel(LOCTEXT("ResultListTypeHeader", "Type")).FillWidth(0.7)
+					+SHeaderRow::Column("Component").DefaultLabel(LOCTEXT("ResultListComponentHeader", "Component")).FillWidth(3)
+					+SHeaderRow::Column("Normal").DefaultLabel(LOCTEXT("ResultListNormalHeader", "Normal")).FillWidth(1.8)
 				)
 			]
 		]
@@ -225,14 +225,14 @@ void SCAQueryDetails::Construct(const FArguments& InArgs, TSharedPtr<SCollisionA
 }
 END_SLATE_FUNCTION_BUILD_OPTIMIZATION
 
-FString SCAQueryDetails::GetStartString() const
+FText SCAQueryDetails::GetStartText() const
 {
-	return bDisplayQuery ? CurrentQuery.Start.ToString() : TEXT("");
+	return bDisplayQuery ? CurrentQuery.Start.ToText() : FText::GetEmpty();
 }
 
-FString SCAQueryDetails::GetEndString() const
+FText SCAQueryDetails::GetEndText() const
 {
-	return bDisplayQuery ? CurrentQuery.End.ToString() : TEXT("");
+	return bDisplayQuery ? CurrentQuery.End.ToText() : FText::GetEmpty();
 }
 
 TSharedRef<ITableRow> SCAQueryDetails::ResultListGenerateRow(TSharedPtr<FCAHitInfo> Info, const TSharedRef<STableViewBase>& OwnerTable)
diff --git a/Engine/Source/Developer/CollisionAnalyzer/Private/SCAQueryDetails.h b/Engine/Source/Developer/CollisionAnalyzer/Private/SCAQueryDetails.h
--- a/Engine/Source/Developer/CollisionAnalyzer/Private/SCAQueryDetails.h
+++ b/Engine/Source/Developer/CollisionAnalyzer/Private/SCAQueryDetails.h
@@ -44,8 +44,8 @@ public:
 	FCAQuery* GetCurrentQuery();
 
 	// Get delegates
-	FString GetStartString() const;
-	FString GetEndString() const;
+	FText GetStartText() const;
+	FText GetEndText() const;
 	ESlateCheckBoxState::Type GetShowMissesState() const;
 	// List delegates
 	TSharedRef<ITableRow> ResultListGenerateRow(TSharedPtr<FCAHitInfo> Info, const TSharedRef<STableViewBase>& OwnerTable);
diff --git a/Engine/Source/Developer/DeviceManager/Private/Widgets/Apps/SDeviceApps.cpp b/Engine/Source/Developer/DeviceManager/Private/Widgets/Apps/SDeviceApps.cpp
--- a/Engine/Source/Developer/DeviceManager/Private/Widgets/Apps/SDeviceApps.cpp
+++ b/Engine/Source/Developer/DeviceManager/Private/Widgets/Apps/SDeviceApps.cpp
@@ -59,13 +59,13 @@ void SDeviceApps::Construct( const FArguments& InArgs, const FDeviceManagerModel
 										SNew(SHeaderRow)
 
 										+ SHeaderRow::Column("Name")
-											.DefaultLabel(LOCTEXT("AppListNameColumnHeader", "Name").ToString())
+											.DefaultLabel(LOCTEXT("AppListNameColumnHeader", "Name"))
 
 										+ SHeaderRow::Column("Date")
-											.DefaultLabel(LOCTEXT("AppListDeploymentDateColumnHeader", "Date deployed").ToString())
+											.DefaultLabel(LOCTEXT("AppListDeploymentDateColumnHeader", "Date deployed"))
 
 										+ SHeaderRow::Column("Owner")
-											.DefaultLabel(LOCTEXT("AppListOwnerColumnHeader", "Deployed by").ToString())
+											.DefaultLabel(LOCTEXT("AppListOwnerColumnHeader", "Deployed by"))
 									)
 							]
 					]
@@ -81,7 +81,7 @@ void SDeviceApps::Construct( const FArguments& InArgs, const FDeviceManagerModel
 					.Visibility(this, &SDeviceApps::HandleSelectDeviceOverlayVisibility)
 					[
 						SNew(STextBlock)
-							.Text(LOCTEXT("SelectSessionOverlayText", "Please select a device from the Device Browser").ToString())
+							.Text(LOCTEXT("SelectSessionOverlayText", "Please select a device from the Device Browser"))
 					]
 			]
 	];
diff --git a/Engine/Source/Developer/DeviceManager/Private/Widgets/Browser/SDeviceBrowserDeviceAdder.cpp b/Engine/Source/Developer/DeviceManager/Private/Widgets/Browser/SDeviceBrowserDeviceAdder.cpp
--- a/Engine/Source/Developer/DeviceManager/Private/Widgets/Browser/SDeviceBrowserDeviceAdder.cpp
+++ b/Engine/Source/Developer/DeviceManager/Private/Widgets/Browser/SDeviceBrowserDeviceAdder.cpp
@@ -39,7 +39,7 @@ void SDeviceBrowserDeviceAdder::Construct( const FArguments& InArgs, const ITarg
 							.HAlign(HAlign_Left)
 							[
 								SNew(STextBlock)
-									.Text(LOCTEXT("PlatformLabel", "Platform:").ToString())
+									.Text(LOCTEXT("PlatformLabel", "Platform:"))
 							]
 
 						+ SVerticalBox::Slot()
@@ -72,7 +72,7 @@ void SDeviceBrowserDeviceAdder::Construct( const FArguments& InArgs, const ITarg
 							.HAlign(HAlign_Left)
 							[
 								SNew(STextBlock)
-									.Text(LOCTEXT("DeviceIdLabel", "Device Identifier:").ToString())
+									.Text(LOCTEXT("DeviceIdLabel", "Device Identifier:"))
 							]
 
 						+ SVerticalBox::Slot()
@@ -96,7 +96,7 @@ void SDeviceBrowserDeviceAdder::Construct( const FArguments& InArgs, const ITarg
 							.HAlign(HAlign_Left)
 							[
 								SNew(STextBlock)
-									.Text(LOCTEXT("DisplayNameLabel", "Display Name:").ToString())
+									.Text(LOCTEXT("DisplayNameLabel", "Display Name:"))
 							]
 
 						+ SVerticalBox::Slot()
@@ -116,7 +116,7 @@ void SDeviceBrowserDeviceAdder::Construct( const FArguments& InArgs, const ITarg
 						SAssignNew(AddButton, SButton)
 							.ContentPadding(FMargin(9.0, 2.0))
 							.IsEnabled(this, &SDeviceBrowserDeviceAdder::HandleAddButtonIsEnabled)
-							.Text(LOCTEXT("AddButtonText", "Add").ToString())
+							.Text(LOCTEXT("AddButtonText", "Add"))
 							.OnClicked(this, &SDeviceBrowserDeviceAdder::HandleAddButtonClicked)
 					]
 			]
@@ -139,7 +139,7 @@ void SDeviceBrowserDeviceAdder::Construct( const FArguments& InArgs, const ITarg
 							.HAlign(HAlign_Left)
 							[
 								SNew(STextBlock)
-									.Text(LOCTEXT("UserNameLabel", "User:").ToString())
+									.Text(LOCTEXT("UserNameLabel", "User:"))
 							]
 
 						+ SVerticalBox::Slot()
@@ -162,7 +162,7 @@ void SDeviceBrowserDeviceAdder::Construct( const FArguments& InArgs, const ITarg
 							.HAlign(HAlign_Left)
 							[
 								SNew(STextBlock)
-									.Text(LOCTEXT("UserPasswordLabel", "Password:").ToString())
+									.Text(LOCTEXT("UserPasswordLabel", "Password:"))
 							]
 
 						+ SVerticalBox::Slot()
diff --git a/Engine/Source/Developer/DeviceManager/Private/Widgets/Browser/SDeviceBrowserFilterBar.cpp b/Engine/Source/Developer/DeviceManager/Private/Widgets/Browser/SDeviceBrowserFilterBar.cpp
--- a/Engine/Source/Developer/DeviceManager/Private/Widgets/Browser/SDeviceBrowserFilterBar.cpp
+++ b/Engine/Source/Developer/DeviceManager/Private/Widgets/Browser/SDeviceBrowserFilterBar.cpp
@@ -54,7 +54,7 @@ void SDeviceBrowserFilterBar::Construct( const FArguments& InArgs, FDeviceBrowse
 					.ButtonContent()
 					[
 						SNew(STextBlock)
-							.Text(LOCTEXT("PlatformComboButtonText", "Platforms").ToString())
+							.Text(LOCTEXT("PlatformComboButtonText", "Platforms"))
 					]
 					.ContentPadding(FMargin(6.0f, 2.0f))
 					.MenuContent()
diff --git a/Engine/Source/Developer/DeviceManager/Private/Widgets/Details/SDeviceDetails.cpp b/Engine/Source/Developer/DeviceManager/Private/Widgets/Details/SDeviceDetails.cpp
--- a/Engine/Source/Developer/DeviceManager/Private/Widgets/Details/SDeviceDetails.cpp
+++ b/Engine/Source/Developer/DeviceManager/Private/Widgets/Details/SDeviceDetails.cpp
@@ -87,7 +87,7 @@ void SDeviceDetails::Construct( const FArguments& InArgs, const FDeviceManagerMo
 					.Visibility(this, &SDeviceDetails::HandleSelectDeviceOverlayVisibility)
 					[
 						SNew(STextBlock)
-							.Text(LOCTEXT("SelectSessionOverlayText", "Please select a device from the Device Browser").ToString())
+							.Text(LOCTEXT("SelectSessionOverlayText", "Please select a device from the Device Browser"))
 					]
 			]
 	];
diff --git a/Engine/Source/Developer/DeviceManager/Private/Widgets/Processes/SDeviceProcesses.cpp b/Engine/Source/Developer/DeviceManager/Private/Widgets/Processes/SDeviceProcesses.cpp
--- a/Engine/Source/Developer/DeviceManager/Private/Widgets/Processes/SDeviceProcesses.cpp
+++ b/Engine/Source/Developer/DeviceManager/Private/Widgets/Processes/SDeviceProcesses.cpp
@@ -60,23 +60,23 @@ void SDeviceProcesses::Construct( const FArguments& InArgs, const FDeviceManager
 										SNew(SHeaderRow)
 
 										+ SHeaderRow::Column("Name")
-											.DefaultLabel(LOCTEXT("ProcessListNameColumnHeader", "Process Name").ToString())
+											.DefaultLabel(LOCTEXT("ProcessListNameColumnHeader", "Process Name"))
 											.FillWidth(0.275f)
 
 										+ SHeaderRow::Column("PID")
-											.DefaultLabel(LOCTEXT("ProcessListPidColumnHeader", "PID").ToString())
+											.DefaultLabel(LOCTEXT("ProcessListPidColumnHeader", "PID"))
 											.FillWidth(0.15f)
 
 										+ SHeaderRow::Column("User")
-											.DefaultLabel(LOCTEXT("ProcessListUserColumnHeader", "User").ToString())
+											.DefaultLabel(LOCTEXT("ProcessListUserColumnHeader", "User"))
 											.FillWidth(0.275f)
 
 										+ SHeaderRow::Column("Threads")
-											.DefaultLabel(LOCTEXT("ProcessListThreadsColumnHeader", "Threads").ToString())
+											.DefaultLabel(LOCTEXT("ProcessListThreadsColumnHeader", "Threads"))
 											.FillWidth(0.15f)
 
 										+ SHeaderRow::Column("Parent")
-											.DefaultLabel(LOCTEXT("ProcessListParentColumnHeader", "Parent PID").ToString())
+											.DefaultLabel(LOCTEXT("ProcessListParentColumnHeader", "Parent PID"))
 											.FillWidth(0.15f)
 									)
 							]
@@ -98,10 +98,10 @@ void SDeviceProcesses::Construct( const FArguments& InArgs, const FDeviceManager
 											.IsChecked(this, &SDeviceProcesses::HandleProcessTreeCheckBoxIsChecked)
 											.OnCheckStateChanged(this, &SDeviceProcesses::HandleProcessTreeCheckBoxCheckStateChanged)
 											.Padding(FMargin(4.0f, 0.0f))
-											.ToolTipText(LOCTEXT("ProcessTreeCheckBoxToolTip", "Check this box to display the list of processes as a tree instead of a flat list").ToString())
+											.ToolTipText(LOCTEXT("ProcessTreeCheckBoxToolTip", "Check this box to display the list of processes as a tree instead of a flat list"))
 											[
 												SNew(STextBlock)
-													.Text(LOCTEXT("ProcessTreeCheckBoxText", "Show process tree").ToString())
+													.Text(LOCTEXT("ProcessTreeCheckBoxText", "Show process tree"))
 											]
 									]
 
@@ -111,7 +111,7 @@ void SDeviceProcesses::Construct( const FArguments& InArgs, const FDeviceManager
 										SNew(SButton)
 											.IsEnabled(this, &SDeviceProcesses::HandleTerminateProcessButtonIsEnabled)
 											.OnClicked(this, &SDeviceProcesses::HandleTerminateProcessButtonClicked)
-											.Text(LOCTEXT("TerminateProcessButtonText", "Terminate Process").ToString())
+											.Text(LOCTEXT("TerminateProcessButtonText", "Terminate Process"))
 									]
 							]
 					]
diff --git a/Engine/Source/Developer/FriendsAndChat/Private/UI/Widgets/SFriendsList.cpp b/Engine/Source/Developer/FriendsAndChat/Private/UI/Widgets/SFriendsList.cpp
--- a/Engine/Source/Developer/FriendsAndChat/Private/UI/Widgets/SFriendsList.cpp
+++ b/Engine/Source/Developer/FriendsAndChat/Private/UI/Widgets/SFriendsList.cpp
@@ -88,7 +88,7 @@ void SFriendsList::Construct( const FArguments& InArgs )
 					[
 						SNew( STextBlock )
 						.Font( FriendStyle.FriendsFontStyle )
-						.Text( LOCTEXT("FriendsList", "Friends").ToString() )
+						.Text( LOCTEXT("FriendsList", "Friends") )
 					]
 				]
 				+SHorizontalBox::Slot()
@@ -103,7 +103,7 @@ void SFriendsList::Construct( const FArguments& InArgs )
 					[
 						SNew(STextBlock)
 						.Font(FriendStyle.FriendsFontStyle)
-						.Text(LOCTEXT("RecentPlayers", "Recent").ToString())
+						.Text(LOCTEXT("RecentPlayers", "Recent"))
 					]
 				]
 				+SHorizontalBox::Slot()
@@ -118,7 +118,7 @@ void SFriendsList::Construct( const FArguments& InArgs )
 					[
 						SNew( STextBlock )
 						.Font( FriendStyle.FriendsFontStyle )
-						.Text( LOCTEXT("FriendRequest", "Requests").ToString() )
+						.Text( LOCTEXT("FriendRequest", "Requests") )
 					]
 				]
 			]
@@ -293,10 +293,11 @@ void SFriendsList::RefreshFriendsList()
 	}
 }
 
-FString SFriendsList::GetListCountText() const
+FText SFriendsList::GetListCountText() const
 {
-	static const FString ListCount = LOCTEXT( "FriendList_Title", "Friends:1 " ).ToString();
-	return ListCount + FString::Printf( TEXT( "%d" ), FFriendsAndChatManager::Get()->GetFriendCount() );
+	FFormatNamedArguments Args;
+	Args.Add(TEXT("NumberOfFriends"), FFriendsAndChatManager::Get()->GetFriendCount());
+	return FText::Format(LOCTEXT( "FriendList_Title", "FRIENDS : {NumberOfFriends}" ), Args);
 }
 
 FReply SFriendsList::SearchFriend_OnClicked()
diff --git a/Engine/Source/Developer/FriendsAndChat/Private/UI/Widgets/SFriendsList.h b/Engine/Source/Developer/FriendsAndChat/Private/UI/Widgets/SFriendsList.h
--- a/Engine/Source/Developer/FriendsAndChat/Private/UI/Widgets/SFriendsList.h
+++ b/Engine/Source/Developer/FriendsAndChat/Private/UI/Widgets/SFriendsList.h
@@ -48,7 +48,7 @@ private:
 	 * Get the list count as text.
 	 * @return The number of friends
 	 */
-	FString GetListCountText() const;
+	FText GetListCountText() const;
 
 	/**
 	 * Handles the search for friend button clicked.
diff --git a/Engine/Source/Developer/LogVisualizer/Private/SLogVisualizer.cpp b/Engine/Source/Developer/LogVisualizer/Private/SLogVisualizer.cpp
--- a/Engine/Source/Developer/LogVisualizer/Private/SLogVisualizer.cpp
+++ b/Engine/Source/Developer/LogVisualizer/Private/SLogVisualizer.cpp
@@ -183,8 +183,8 @@ void SLogVisualizer::Construct(const FArguments& InArgs, FLogVisualizer* InLogVi
 						.Content()
 						[
 							SNew(STextBlock)
-							.Text(LOCTEXT("VisLogDrawLogsPath", "Draw Log\'s path").ToString())
-							.ToolTipText(LOCTEXT("VisLogDrawLogsPathTooltip", "Toggle whether path of composed of log entries\' locations").ToString())
+							.Text(LOCTEXT("VisLogDrawLogsPath", "Draw Log\'s path"))
+							.ToolTipText(LOCTEXT("VisLogDrawLogsPathTooltip", "Toggle whether path of composed of log entries\' locations"))
 						]
 					]
 					+SVerticalBox::Slot()
@@ -196,8 +196,8 @@ void SLogVisualizer::Construct(const FArguments& InArgs, FLogVisualizer* InLogVi
 						.Content()
 						[
 							SNew(STextBlock)
-							.Text(LOCTEXT("VisLogIgnoreTrivialLogs", "Ignore trivial logs").ToString())
-							.ToolTipText(LOCTEXT("VisLogIgnoreTrivialLogsTooltip", "Whether to show trivial logs, i.e. the ones with only one entry.").ToString())
+							.Text(LOCTEXT("VisLogIgnoreTrivialLogs", "Ignore trivial logs"))
+							.ToolTipText(LOCTEXT("VisLogIgnoreTrivialLogsTooltip", "Whether to show trivial logs, i.e. the ones with only one entry."))
 						]
 					]
 				]
@@ -224,7 +224,7 @@ void SLogVisualizer::Construct(const FArguments& InArgs, FLogVisualizer* InLogVi
 						.HeaderRow(
 																																																																																				SNew(SHeaderRow)
 							// ID
-							+SHeaderRow::Column(*NAME_LogName.ToString())
+							+SHeaderRow::Column(NAME_LogName)
 							.SortMode(this, &SLogVisualizer::GetLogsSortMode)
 							.OnSort(this, &SLogVisualizer::OnSortByChanged)
 							.HAlignCell(HAlign_Left)
@@ -237,7 +237,7 @@ void SLogVisualizer::Construct(const FArguments& InArgs, FLogVisualizer* InLogVi
 								.Padding(0.0, 2.0)
 								[
 									SNew(STextBlock)
-									.Text(LOCTEXT("VisLogName", "Log Subject").ToString())
+									.Text(LOCTEXT("VisLogName", "Log Subject"))
 								]
 								+SHorizontalBox::Slot()
 								.AutoWidth()
@@ -250,7 +250,7 @@ void SLogVisualizer::Construct(const FArguments& InArgs, FLogVisualizer* InLogVi
 									.RevertTextOnEscape(true)
 								]
 							]
-							+SHeaderRow::Column(*NAME_LogTimeSpan.ToString())
+							+SHeaderRow::Column(NAME_LogTimeSpan)
 							/*.OnSort(this, &SLogVisualizer::OnSortByChanged)*/
 							.VAlignCell(VAlign_Center)
 							[
@@ -259,8 +259,8 @@ void SLogVisualizer::Construct(const FArguments& InArgs, FLogVisualizer* InLogVi
 								.AutoHeight()
 								[
 									SNew(STextBlock)
-									.Text(LOCTEXT("VisLogTimeSpan", "Overview").ToString())
-									.ToolTipText(LOCTEXT("VisLogTimeSpanTooltip", "Mouse-over to see timestamp, click to show log entry").ToString())
+									.Text(LOCTEXT("VisLogTimeSpan", "Overview"))
+									.ToolTipText(LOCTEXT("VisLogTimeSpanTooltip", "Mouse-over to see timestamp, click to show log entry"))
 								]
 							]
 						)
diff --git a/Engine/Source/Developer/Profiler/Private/Widgets/SEventGraph.cpp b/Engine/Source/Developer/Profiler/Private/Widgets/SEventGraph.cpp
--- a/Engine/Source/Developer/Profiler/Private/Widgets/SEventGraph.cpp
+++ b/Engine/Source/Developer/Profiler/Private/Widgets/SEventGraph.cpp
@@ -1004,7 +1004,7 @@ TSharedRef<SWidget> SEventGraph::GetWidgetForEventGraphTypes()
 		.Padding( 2.0f )
 		[	
 			SNew(STextBlock)
-			.Text( LOCTEXT("Toolbar_Type","Type").ToString() )
+			.Text( LOCTEXT("Toolbar_Type","Type") )
 			.TextStyle( FEditorStyle::Get(), TEXT("Profiler.CaptionBold") )
 		]
 
@@ -1058,7 +1058,7 @@ TSharedRef<SWidget> SEventGraph::GetWidgetForEventGraphViewModes()
 		.Padding( 2.0f )
 		[	
 			SNew(STextBlock)
-			.Text( LOCTEXT("Toolbar_ViewMode","View mode").ToString() )
+			.Text( LOCTEXT("Toolbar_ViewMode","View mode") )
 			.TextStyle( FEditorStyle::Get(), TEXT("Profiler.CaptionBold") )
 		]
 
@@ -1573,15 +1573,15 @@ TSharedPtr<SWidget> SEventGraph::EventGraph_GetMenuContent() const
 
 // 		MenuBuilder.AddMenuEntry
 // 		( 
-// 			LOCTEXT("ContextMenu_Root_Reset", "Reset To Default").ToString(), 
-// 			LOCTEXT("ContextMenu_Root_Reset_Desc", "Resets the root options and removes from the history").ToString(), 
+// 			LOCTEXT("ContextMenu_Root_Reset", "Reset To Default"), 
+// 			LOCTEXT("ContextMenu_Root_Reset_Desc", "Resets the root options and removes from the history"), 
 // 			TEXT("Profiler.Misc.ResetToDefault"), ResetRoot_Custom(), NAME_None, EUserInterfaceActionType::Button 
 // 		);
 	}
 	//MenuBuilder.EndSection();
 
 	// Culling menu
-	//MenuBuilder.BeginSection("Culling", LOCTEXT("ContextMenu_Culling","Culling").ToString() );
+	//MenuBuilder.BeginSection("Culling", LOCTEXT("ContextMenu_Culling","Culling") );
 	{
 		FText CullingDesc;
 		
@@ -2928,7 +2928,7 @@ FText SEventGraph::HistoryBack_GetToolTipText() const
 	// TODO: Add a nicer custom tooltip.
 	if( HistoryBack_IsEnabled() )
 	{
-		return FText::Format( LOCTEXT("HistoryBack_Tooltip", "Back to {0}"), FText::FromString( EventGraphStatesHistory[CurrentStateIndex-1]->GetFullDescription() ) );
+		return FText::Format( LOCTEXT("HistoryBack_Tooltip", "Back to {0}"), EventGraphStatesHistory[CurrentStateIndex-1]->GetFullDescription() );
 	}
 
 	return FText::GetEmpty();
@@ -2950,7 +2950,7 @@ FText SEventGraph::HistoryForward_GetToolTipText() const
 	// TODO: Add a nicer custom tooltip.
 	if( HistoryForward_IsEnabled() )
 	{
-		return FText::Format( LOCTEXT("HistoryForward_Tooltip", "Forward to {0}"), FText::FromString( EventGraphStatesHistory[CurrentStateIndex+1]->GetFullDescription() ) );
+		return FText::Format( LOCTEXT("HistoryForward_Tooltip", "Forward to {0}"), EventGraphStatesHistory[CurrentStateIndex+1]->GetFullDescription() );
 	}
 
 	return FText::GetEmpty();
@@ -2975,7 +2975,7 @@ TSharedRef<SWidget> SEventGraph::HistoryList_GetMenuContent() const
 		const FEventGraphStateRef StateRef = EventGraphStatesHistory[StateIndex];
 		MenuBuilder.AddMenuEntry
 		(
-			FText::FromString( StateRef->GetFullDescription() ),
+			StateRef->GetFullDescription(),
 			FText(),
 			FSlateIcon(),
 			HistoryList_GoTo_Custom( StateIndex ),
@@ -3354,63 +3354,72 @@ bool SEventGraph::SelectAllFrames_CanExecute() const
 
 #define LOCTEXT_NAMESPACE "FEventGraphState"
 
-FString SEventGraph::FEventGraphState::GetFullDescription() const
+FText SEventGraph::FEventGraphState::GetFullDescription() const
 {
-	FString Result = FString::Printf( *LOCTEXT("FullDesc","Event graph with range {%i,%i} contains %i frame(s)").ToString(), 
-		GetEventGraph()->GetFrameStartIndex(), 
-		GetEventGraph()->GetFrameEndIndex(), 
-		GetNumFrames() );
+	FTextBuilder Builder;
+
+	FFormatNamedArguments Args;
+	Args.Add(TEXT("FrameStartIndex"), GetEventGraph()->GetFrameStartIndex());
+	Args.Add(TEXT("FrameEndIndex"), GetEventGraph()->GetFrameEndIndex());
+	Args.Add(TEXT("NumberOfFrames"), GetNumFrames());
+	Builder.AppendLineFormat(LOCTEXT("FullDesc", "Event graph with range ({FrameStartIndex},{FrameEndIndex}) contains {NumberOfFrames} frame(s)"), Args);
+
+	Builder.Indent();
 
 	if( IsRooted() )
 	{
-		Result += LINE_TERMINATOR;
-		Result += TEXT("  ");
-		Result += GetRootedDesc();
+		Builder.AppendLine(GetRootedDesc());
 	}
 
 	if( IsCulled() )
 	{
-		Result += LINE_TERMINATOR;
-		Result += TEXT("  ");
-		Result += GetCullingDesc();
+		Builder.AppendLine(GetCullingDesc());
 	}
 
 	if( IsFiltered() )
 	{
-		Result += LINE_TERMINATOR;
-		Result += TEXT("  ");
-		Result += GetFilteringDesc();
+		Builder.AppendLine(GetFilteringDesc());
 	}
 
-	return Result;
+	return Builder.ToText();
 }
 
-FString SEventGraph::FEventGraphState::GetRootedDesc() const
+FText SEventGraph::FEventGraphState::GetRootedDesc() const
 {
 	const int32 NumFakeRoots = FakeRoot->GetChildren().Num();
-	const FString Result = FString::Printf( *LOCTEXT("RootedDesc","Rooted: %s").ToString(), NumFakeRoots == 1 ? *FakeRoot->GetChildren()[0]->_StatName.GetPlainNameString() : TEXT("Multiple") );
-	return Result;
+
+	if (NumFakeRoots == 1)
+	{
+		FFormatNamedArguments Args;
+		Args.Add(TEXT("StatName"), FText::FromName(FakeRoot->GetChildren()[0]->_StatName));
+		return FText::Format(LOCTEXT("RootedDesc_SingleChild", "Rooted: {StatName}"), Args);
+	}
+
+	return LOCTEXT("RootedDesc_MultipleChildren", "Rooted: Multiple");
 }
 
-FString SEventGraph::FEventGraphState::GetCullingDesc() const
+FText SEventGraph::FEventGraphState::GetCullingDesc() const
 {
-	const FString Result = FString::Printf( *LOCTEXT("CulledDesc","Culled: %s %s").ToString(), 
-		*CullPropertyName.GetPlainNameString(), 
-		*CullEventPtr->GetFormattedValue(FEventGraphSample::GetEventPropertyByName(CullPropertyName).Index) );
-	return Result;
+	FFormatNamedArguments Args;
+	Args.Add(TEXT("CulledPropertyName"), FText::FromName(CullPropertyName));
+	Args.Add(TEXT("EventName"), FText::FromString(CullEventPtr->GetFormattedValue(FEventGraphSample::GetEventPropertyByName(CullPropertyName).Index)));
+
+	return FText::Format(LOCTEXT("CulledDesc", "Culled: {CulledPropertyName} {EventName}"), Args);
 }
 
-FString SEventGraph::FEventGraphState::GetFilteringDesc() const
+FText SEventGraph::FEventGraphState::GetFilteringDesc() const
 {
-	const FString Result = FString::Printf( *LOCTEXT("FilteredDesc","Filtered: %s %s").ToString(), 
-		*FilterPropertyName.GetPlainNameString(), 
-		*FilterEventPtr->GetFormattedValue(FEventGraphSample::GetEventPropertyByName(FilterPropertyName).Index) );
-	return Result;
+	FFormatNamedArguments Args;
+	Args.Add(TEXT("FilterPropertyName"), FText::FromName(FilterPropertyName));
+	Args.Add(TEXT("EventName"), FText::FromString(FilterEventPtr->GetFormattedValue(FEventGraphSample::GetEventPropertyByName(FilterPropertyName).Index)));
+
+	return FText::Format(LOCTEXT("FilteredDesc", "Filtered: {FilterPropertyName} {EventName}"), Args);
 }
 
-FString SEventGraph::FEventGraphState::GetHistoryDesc() const
+FText SEventGraph::FEventGraphState::GetHistoryDesc() const
 {
-	FString Result(LOCTEXT("DefaultDesc","Default state").ToString());
+	FText Result = LOCTEXT("DefaultDesc", "Default state");
+
 	if( HistoryType == EEventHistoryTypes::Rooted )
 	{
 		Result = GetRootedDesc();
@@ -3423,6 +3432,7 @@ FString SEventGraph::FEventGraphState::GetHistoryDesc() const
 	{
 		Result = GetFilteringDesc();
 	}
+
 	return Result;
 }
 
diff --git a/Engine/Source/Developer/Profiler/Private/Widgets/SEventGraph.h b/Engine/Source/Developer/Profiler/Private/Widgets/SEventGraph.h
--- a/Engine/Source/Developer/Profiler/Private/Widgets/SEventGraph.h
+++ b/Engine/Source/Developer/Profiler/Private/Widgets/SEventGraph.h
@@ -747,15 +747,15 @@ protected:
 			return AverageEventGraph->GetNumFrames();
 		}
 
-		FString GetFullDescription() const;
+		FText GetFullDescription() const;
 
-		FString GetRootedDesc() const;
+		FText GetRootedDesc() const;
 
-		FString GetCullingDesc() const;
+		FText GetCullingDesc() const;
 
-		FString GetFilteringDesc() const;
+		FText GetFilteringDesc() const;
 
-		FString GetHistoryDesc() const;
+		FText GetHistoryDesc() const;
 
 		const FEventGraphDataRef& GetEventGraph() const
 		{
diff --git a/Engine/Source/Developer/Profiler/Private/Widgets/SEventGraphTooltip.cpp b/Engine/Source/Developer/Profiler/Private/Widgets/SEventGraphTooltip.cpp
--- a/Engine/Source/Developer/Profiler/Private/Widgets/SEventGraphTooltip.cpp
+++ b/Engine/Source/Developer/Profiler/Private/Widgets/SEventGraphTooltip.cpp
@@ -17,17 +17,8 @@ TSharedPtr<SToolTip> SEventGraphTooltip::GetTableCellTooltip( const FEventGraphS
 	const float Alpha = EventSample->_FramePct * 0.01f;
 	const FLinearColor ColorAndOpacity = FMath::Lerp(DefaultColor,ThreadColor,Alpha);
 
-	const FString InclusiveTimePctCaller = FString::Printf
-	( 
-		*LOCTEXT("PctOfTheCaller", "(%s of the caller)").ToString(), 
-		*EventSample->GetFormattedValue( EEventPropertyIndex::InclusiveTimePct ) 
-	);
-
-	const FString ExclusiveTimePctCaller = FString::Printf
-	( 
-		*LOCTEXT("PctOfThisCallIncTime", "(%s of this call's inc time)").ToString(), 
-		*EventSample->GetFormattedValue( EEventPropertyIndex::ExclusiveTimePct ) 
-	);
+	const FText InclusiveTimePctCaller = FText::Format(LOCTEXT("PctOfTheCaller", "({0} of the caller)"), FText::FromString(EventSample->GetFormattedValue(EEventPropertyIndex::InclusiveTimePct)));
+	const FText ExclusiveTimePctCaller = FText::Format(LOCTEXT("PctOfThisCallIncTime", "({0} of this call's inc time)"), FText::FromString(EventSample->GetFormattedValue(EEventPropertyIndex::ExclusiveTimePct)));
 
 	TAttribute< EVisibility > HotPathIconVisibility = TAttribute< EVisibility >::Create( TAttribute<EVisibility>::FGetter::CreateStatic( &SEventGraphTooltip::GetHotPathIconVisibility, EventSample ) );
 
@@ -82,7 +73,7 @@ TSharedPtr<SToolTip> SEventGraphTooltip::GetTableCellTooltip( const FEventGraphS
 							.VAlign(VAlign_Center)
 							[
 								SNew(STextBlock)
-								.Text( LOCTEXT("HotPathInformation", "Hot path, should be investigated for bottlenecks").ToString() )
+								.Text( LOCTEXT("HotPathInformation", "Hot path, should be investigated for bottlenecks") )
 								//.TextStyle( FEditorStyle::Get(), TEXT("Profiler.TooltipBold") )
 								.ColorAndOpacity( ColorAndOpacity )
 								.Font( FSlateFontInfo( FPaths::EngineContentDir() / TEXT("Slate/Fonts/Roboto-Regular.ttf"), 8 ) )
@@ -113,14 +104,14 @@ TSharedPtr<SToolTip> SEventGraphTooltip::GetTableCellTooltip( const FEventGraphS
 						.Padding( 2.0f )
 						[
 							SNew(STextBlock)
-							.Text(LOCTEXT("TT_Thread", "Thread:").ToString())
+							.Text(LOCTEXT("TT_Thread", "Thread:"))
 							.TextStyle( FEditorStyle::Get(), TEXT("Profiler.TooltipBold") )
 						]
 						+SGridPanel::Slot(1,0)
 							.Padding( 2.0f )
 							[
 								SNew(STextBlock)
-								.Text( EventSample->_ThreadName.ToString() )
+								.Text( FText::FromName(EventSample->_ThreadName ))
 								.TextStyle( FEditorStyle::Get(), TEXT("Profiler.Tooltip") )
 							]
 
@@ -129,7 +120,7 @@ TSharedPtr<SToolTip> SEventGraphTooltip::GetTableCellTooltip( const FEventGraphS
 							.Padding( 2.0f )
 							[
 								SNew(STextBlock)
-								.Text(LOCTEXT("TT_Event", "Event:").ToString())
+								.Text(LOCTEXT("TT_Event", "Event:"))
 								.TextStyle( FEditorStyle::Get(), TEXT("Profiler.TooltipBold") )
 							]
 						+SGridPanel::Slot(3,0)
@@ -145,14 +136,14 @@ TSharedPtr<SToolTip> SEventGraphTooltip::GetTableCellTooltip( const FEventGraphS
 							.Padding( 2.0f )
 							[
 								SNew(STextBlock)
-								.Text(LOCTEXT("TT_Group", "Group:").ToString())
+								.Text(LOCTEXT("TT_Group", "Group:"))
 								.TextStyle( FEditorStyle::Get(), TEXT("Profiler.TooltipBold") )
 							]
 						+SGridPanel::Slot(5,0)
 							.Padding( 2.0f )
 							[
 								SNew(STextBlock)
-								.Text( EventSample->_GroupName.ToString() )
+								.Text( FText::FromName(EventSample->_GroupName) )
 								.TextStyle( FEditorStyle::Get(), TEXT("Profiler.Tooltip") )
 							]
 
@@ -163,7 +154,7 @@ TSharedPtr<SToolTip> SEventGraphTooltip::GetTableCellTooltip( const FEventGraphS
 							.Padding( 2.0f )
 							[
 								SNew(STextBlock)
-								.Text(LOCTEXT("TT_InclusiveTime", "Inclusive time:").ToString())
+								.Text(LOCTEXT("TT_InclusiveTime", "Inclusive time:"))
 								.TextStyle( FEditorStyle::Get(), TEXT("Profiler.TooltipBold") )
 							]
 						// [InclusiveTimeMS] ms
@@ -187,7 +178,7 @@ TSharedPtr<SToolTip> SEventGraphTooltip::GetTableCellTooltip( const FEventGraphS
 							.Padding( 2.0f )
 							[
 								SNew(STextBlock)
-								.Text(LOCTEXT("TT_AvgIncTimePerCall", "Avg inc time per call:").ToString())
+								.Text(LOCTEXT("TT_AvgIncTimePerCall", "Avg inc time per call:"))
 								.TextStyle( FEditorStyle::Get(), TEXT("Profiler.TooltipBold") )
 							]
 						// [_AvgInclusiveTimePerCallMS] ms
@@ -207,7 +198,7 @@ TSharedPtr<SToolTip> SEventGraphTooltip::GetTableCellTooltip( const FEventGraphS
 							.Padding( 2.0f )
 							[
 								SNew(STextBlock)
-								.Text(LOCTEXT("TT_PctOfThread", "% of thread:").ToString())
+								.Text(LOCTEXT("TT_PctOfThread", "% of thread:"))
 								.TextStyle( FEditorStyle::Get(), TEXT("Profiler.TooltipBold") )
 							]
 						// [ThreadPct] %
@@ -223,7 +214,7 @@ TSharedPtr<SToolTip> SEventGraphTooltip::GetTableCellTooltip( const FEventGraphS
 							.Padding( 2.0f )
 							[
 								SNew(STextBlock)
-								.Text(LOCTEXT("TT_PctOfFrame", "% of frame:").ToString())
+								.Text(LOCTEXT("TT_PctOfFrame", "% of frame:"))
 								.TextStyle( FEditorStyle::Get(), TEXT("Profiler.TooltipBold") )
 							]
 						// [FramePct] %
@@ -243,7 +234,7 @@ TSharedPtr<SToolTip> SEventGraphTooltip::GetTableCellTooltip( const FEventGraphS
 							.Padding( 2.0f )
 							[
 								SNew(STextBlock)
-								.Text(LOCTEXT("TT_MinMaxAvgIncTime", "Min/Max/Avg inclusive time:").ToString())
+								.Text(LOCTEXT("TT_MinMaxAvgIncTime", "Min/Max/Avg inclusive time:"))
 								.TextStyle( FEditorStyle::Get(), TEXT("Profiler.TooltipBold") )
 							]
 						// [_MinInclusiveTimeMS] ms
@@ -279,7 +270,7 @@ TSharedPtr<SToolTip> SEventGraphTooltip::GetTableCellTooltip( const FEventGraphS
 							.Padding( 2.0f )
 							[
 								SNew(STextBlock)
-								.Text(LOCTEXT("TT_ExclusiveTime", "Exclusive time:").ToString())
+								.Text(LOCTEXT("TT_ExclusiveTime", "Exclusive time:"))
 								.TextStyle( FEditorStyle::Get(), TEXT("Profiler.TooltipBold") )
 							]
 						// [_ExclusiveTimeMS] ms
@@ -307,7 +298,7 @@ TSharedPtr<SToolTip> SEventGraphTooltip::GetTableCellTooltip( const FEventGraphS
 							.Padding( 2.0f )
 							[
 								SNew(STextBlock)
-								.Text(LOCTEXT("TT_NulCallsPerFrame", "Num calls per frame:").ToString())
+								.Text(LOCTEXT("TT_NulCallsPerFrame", "Num calls per frame:"))
 								.TextStyle( FEditorStyle::Get(), TEXT("Profiler.TooltipBold") )
 							]
 						// [_NumCallsPerFrame]
@@ -324,7 +315,7 @@ TSharedPtr<SToolTip> SEventGraphTooltip::GetTableCellTooltip( const FEventGraphS
 							.Padding( 2.0f )
 							[
 								SNew(STextBlock)
-								.Text(LOCTEXT("TT_AvgNumCallsPerFrame", "Avg num calls per frame:").ToString())
+								.Text(LOCTEXT("TT_AvgNumCallsPerFrame", "Avg num calls per frame:"))
 								.TextStyle( FEditorStyle::Get(), TEXT("Profiler.TooltipBold") )
 							]
 						// [_AvgNumCallsPerFrame]
@@ -344,7 +335,7 @@ TSharedPtr<SToolTip> SEventGraphTooltip::GetTableCellTooltip( const FEventGraphS
 							.Padding( 2.0f )
 							[
 								SNew(STextBlock)
-								.Text(LOCTEXT("TT_ThreadDuration", "Thread duration:").ToString())
+								.Text(LOCTEXT("TT_ThreadDuration", "Thread duration:"))
 								.TextStyle( FEditorStyle::Get(), TEXT("Profiler.TooltipBold") )
 							]
 						// [_ThreadDurationMS]
@@ -361,7 +352,7 @@ TSharedPtr<SToolTip> SEventGraphTooltip::GetTableCellTooltip( const FEventGraphS
 							.Padding( 2.0f )
 							[
 								SNew(STextBlock)
-								.Text(LOCTEXT("TT_FrameDuration", "Frame duration:").ToString())
+								.Text(LOCTEXT("TT_FrameDuration", "Frame duration:"))
 								.TextStyle( FEditorStyle::Get(), TEXT("Profiler.TooltipBold") )
 							]
 						// [_FrameDurationMS]
@@ -379,7 +370,7 @@ TSharedPtr<SToolTip> SEventGraphTooltip::GetTableCellTooltip( const FEventGraphS
 							.Padding( 2.0f )
 							[
 								SNew(STextBlock)
-								.Text(LOCTEXT("ThreadToFrame", "Thread to Frame:").ToString())
+								.Text(LOCTEXT("ThreadToFrame", "Thread to Frame:"))
 								.TextStyle( FEditorStyle::Get(), TEXT("Profiler.TooltipBold") )
 							]
 						// [_FrameDurationMS]
diff --git a/Engine/Source/Developer/Profiler/Private/Widgets/SFiltersAndPresets.cpp b/Engine/Source/Developer/Profiler/Private/Widgets/SFiltersAndPresets.cpp
--- a/Engine/Source/Developer/Profiler/Private/Widgets/SFiltersAndPresets.cpp
+++ b/Engine/Source/Developer/Profiler/Private/Widgets/SFiltersAndPresets.cpp
@@ -80,7 +80,7 @@ public:
 		else
 		{
 			return SNew(SToolTip)
-				.Text( LOCTEXT("NotImplemented","Tooltip for multiple profiler instances has not been implemented yet").ToString() );
+				.Text( LOCTEXT("NotImplemented","Tooltip for multiple profiler instances has not been implemented yet") );
 		}
 	}
 
@@ -94,7 +94,7 @@ protected:
 		[
 			SNew(STextBlock)
 			.TextStyle( FEditorStyle::Get(), TEXT("Profiler.TooltipBold") )
-			.Text( LOCTEXT("NoStatData","There is no data for this stat").ToString() )
+			.Text( LOCTEXT("NoStatData","There is no data for this stat") )
 		];
 		RowPos++;
 	}
@@ -109,7 +109,7 @@ protected:
 		[
 			SNew(STextBlock)
 			.TextStyle( FEditorStyle::Get(), TEXT("Profiler.TooltipBold") )
-			.Text( LOCTEXT("StatInstance","Stat information for profiler instance").ToString() )
+			.Text( LOCTEXT("StatInstance","Stat information for profiler instance") )
 		];
 
 		Grid->AddSlot( 0, RowPos++ )
@@ -135,7 +135,7 @@ protected:
 		[
 			SNew(STextBlock)
 			.TextStyle( FEditorStyle::Get(), TEXT("Profiler.TooltipBold") )
-			.Text( LOCTEXT("GroupDesc","Group:").ToString() )
+			.Text( LOCTEXT("GroupDesc","Group:") )
 		];
 
 		Grid->AddSlot( 1, RowPos )
@@ -144,7 +144,7 @@ protected:
 		[
 			SNew(STextBlock)
 			.TextStyle( FEditorStyle::Get(), TEXT("Profiler.Tooltip") )
-			.Text( ProfilerStat.OwningGroup().Name().GetPlainNameString() )
+			.Text( FText::FromName(ProfilerStat.OwningGroup().Name() ))
 		];
 		RowPos++;
 
@@ -153,7 +153,7 @@ protected:
 		[
 			SNew(STextBlock)
 			.TextStyle( FEditorStyle::Get(), TEXT("Profiler.TooltipBold") )
-			.Text( LOCTEXT("NameDesc","Name:").ToString() )
+			.Text( LOCTEXT("NameDesc","Name:") )
 		];
 
 		Grid->AddSlot( 1, RowPos )
@@ -162,7 +162,7 @@ protected:
 		[
 			SNew(STextBlock)
 			.TextStyle( FEditorStyle::Get(), TEXT("Profiler.Tooltip") )
-			.Text( ProfilerStat.Name().GetPlainNameString() )
+			.Text( FText::FromName(ProfilerStat.Name()) )
 		];
 		RowPos++;
 
@@ -171,7 +171,7 @@ protected:
 		[
 			SNew(STextBlock)
 			.TextStyle( FEditorStyle::Get(), TEXT("Profiler.TooltipBold") )
-			.Text( LOCTEXT("TypeDesc","Type:").ToString() )
+			.Text( LOCTEXT("TypeDesc","Type:") )
 		];
 
 		Grid->AddSlot( 1, RowPos )
@@ -210,7 +210,7 @@ protected:
 		[
 			SNew(STextBlock)
 			.TextStyle( FEditorStyle::Get(), TEXT("Profiler.TooltipBold") )
-			.Text( LOCTEXT("ValueDesc","Value").ToString() )
+			.Text( LOCTEXT("ValueDesc","Value") )
 		];
 		RowPos++;
 
@@ -219,7 +219,7 @@ protected:
 		[
 			SNew(STextBlock)
 			.TextStyle( FEditorStyle::Get(), TEXT("Profiler.Tooltip") )
-			.Text( FString(LOCTEXT("MinDesc","Min: ").ToString()) += Aggregated.GetFormattedValue(FProfilerAggregatedStat::EMinValue) )
+			.Text(FText::Format(LOCTEXT("MinDesc", "Min: {0}"), FText::FromString(Aggregated.GetFormattedValue(FProfilerAggregatedStat::EMinValue))))
 		];
 
 		Grid->AddSlot( 1, RowPos )
@@ -227,7 +227,7 @@ protected:
 		[
 			SNew(STextBlock)
 			.TextStyle( FEditorStyle::Get(), TEXT("Profiler.Tooltip") )
-			.Text( FString(LOCTEXT("AvgDesc","Avg: ").ToString()) += Aggregated.GetFormattedValue(FProfilerAggregatedStat::EAvgValue) )
+			.Text(FText::Format(LOCTEXT("AvgDesc", "Avg: {0}"), FText::FromString(Aggregated.GetFormattedValue(FProfilerAggregatedStat::EAvgValue))))
 		];
 
 		Grid->AddSlot( 2, RowPos )
@@ -235,7 +235,7 @@ protected:
 		[
 			SNew(STextBlock)
 			.TextStyle( FEditorStyle::Get(), TEXT("Profiler.Tooltip") )
-			.Text( FString(LOCTEXT("MaxDesc","Max: ").ToString()) += Aggregated.GetFormattedValue(FProfilerAggregatedStat::EMaxValue) )
+			.Text(FText::Format(LOCTEXT("MaxDesc", "Max: {0}"), FText::FromString(Aggregated.GetFormattedValue(FProfilerAggregatedStat::EMaxValue))))
 		];
 		RowPos++;
 
@@ -255,7 +255,7 @@ protected:
 		[
 			SNew(STextBlock)
 			.TextStyle( FEditorStyle::Get(), TEXT("Profiler.TooltipBold") )
-			.Text( FString(LOCTEXT("CallsFramesPctDesc","Calls Frames with call: ").ToString()) += Aggregated.GetFormattedValue(FProfilerAggregatedStat::EFramesWithCallPct) )
+			.Text(FText::Format(LOCTEXT("CallsFramesPctDesc", "Calls Frames with call: {0}"), FText::FromString(Aggregated.GetFormattedValue(FProfilerAggregatedStat::EFramesWithCallPct))))
 		];
 		RowPos++;
 
@@ -264,7 +264,7 @@ protected:
 		[
 			SNew(STextBlock)
 			.TextStyle( FEditorStyle::Get(), TEXT("Profiler.Tooltip") )
-			.Text( FString(LOCTEXT("MinDesc","Min: ").ToString()) += Aggregated.GetFormattedValue(FProfilerAggregatedStat::EMinNumCalls) )
+			.Text(FText::Format(LOCTEXT("MinDesc", "Min: {0}"), FText::FromString(Aggregated.GetFormattedValue(FProfilerAggregatedStat::EMinNumCalls))))
 		];
 
 		Grid->AddSlot( 1, RowPos )
@@ -272,7 +272,7 @@ protected:
 		[
 			SNew(STextBlock)
 			.TextStyle( FEditorStyle::Get(), TEXT("Profiler.Tooltip") )
-			.Text( FString(LOCTEXT("AvgDesc","Avg: ").ToString()) += Aggregated.GetFormattedValue(FProfilerAggregatedStat::EAvgNumCalls) )
+			.Text(FText::Format(LOCTEXT("AvgDesc", "Avg: {0}"), FText::FromString(Aggregated.GetFormattedValue(FProfilerAggregatedStat::EAvgNumCalls))))
 		];
 
 		Grid->AddSlot( 2, RowPos )
@@ -280,7 +280,7 @@ protected:
 		[
 			SNew(STextBlock)
 			.TextStyle( FEditorStyle::Get(), TEXT("Profiler.Tooltip") )
-			.Text( FString(LOCTEXT("MaxDesc","Max: ").ToString()) += Aggregated.GetFormattedValue(FProfilerAggregatedStat::EMaxNumCalls) )
+			.Text(FText::Format(LOCTEXT("MaxDesc", "Max: {0}"), FText::FromString(Aggregated.GetFormattedValue(FProfilerAggregatedStat::EMaxNumCalls))))
 		];
 		RowPos++;
 
@@ -305,29 +305,29 @@ protected:
 	EStatGroupingOrSortingMode
 -----------------------------------------------------------------------------*/
 
-FString EStatGroupingOrSortingMode::ToName( const Type StatGroupingOrSortingMode )
+FText EStatGroupingOrSortingMode::ToName( const Type StatGroupingOrSortingMode )
 {
 	switch( StatGroupingOrSortingMode )
 	{
-	case GroupName: return LOCTEXT("GroupingOrSorting_Name_GroupName", "Group Name").ToString();
-	case StatName: return LOCTEXT("GroupingOrSorting_Name_StatName", "Stat Name").ToString();
-	case StatType: return LOCTEXT("GroupingOrSorting_Name_StatType", "Stat Type").ToString();
-	case StatValue: return LOCTEXT("GroupingOrSorting_Name_StatValue", "Stat Value").ToString();
+	case GroupName: return LOCTEXT("GroupingOrSorting_Name_GroupName", "Group Name");
+	case StatName: return LOCTEXT("GroupingOrSorting_Name_StatName", "Stat Name");
+	case StatType: return LOCTEXT("GroupingOrSorting_Name_StatType", "Stat Type");
+	case StatValue: return LOCTEXT("GroupingOrSorting_Name_StatValue", "Stat Value");
 
-	default: return LOCTEXT("InvalidOrMax", "InvalidOrMax").ToString();
+	default: return LOCTEXT("InvalidOrMax", "InvalidOrMax");
 	}
 }
 
-FString EStatGroupingOrSortingMode::ToDescription( const Type StatGroupingOrSortingMode )
+FText EStatGroupingOrSortingMode::ToDescription(const Type StatGroupingOrSortingMode)
 {
 	switch( StatGroupingOrSortingMode )
 	{
-	case GroupName: return LOCTEXT("GroupingOrSorting_Desc_GroupName", "Creates groups based on stat metadata groups").ToString();
-	case StatName: return LOCTEXT("GroupingOrSorting_Desc_StatName", "Creates one group for one letter").ToString();
-	case StatType: return LOCTEXT("GroupingOrSorting_Desc_StatType", "Creates one group for each stat type").ToString();
-	case StatValue: return LOCTEXT("GroupingOrSorting_Desc_StatValue", "Creates one group for each logarithmic range ie. 0.001 - 0.01, 0.01 - 0.1, 0.1 - 1.0, 1.0 - 10.0 etc").ToString();
+	case GroupName: return LOCTEXT("GroupingOrSorting_Desc_GroupName", "Creates groups based on stat metadata groups");
+	case StatName: return LOCTEXT("GroupingOrSorting_Desc_StatName", "Creates one group for one letter");
+	case StatType: return LOCTEXT("GroupingOrSorting_Desc_StatType", "Creates one group for each stat type");
+	case StatValue: return LOCTEXT("GroupingOrSorting_Desc_StatValue", "Creates one group for each logarithmic range ie. 0.001 - 0.01, 0.01 - 0.1, 0.1 - 1.0, 1.0 - 10.0 etc");
 
-	default: return LOCTEXT("InvalidOrMax", "InvalidOrMax").ToString();
+	default: return LOCTEXT("InvalidOrMax", "InvalidOrMax");
 	}
 }
 
@@ -1170,12 +1170,12 @@ TSharedRef<SWidget> SFiltersAndPresets::GroupBy_OnGenerateWidget( TSharedPtr<ESt
 		.ToolTipText( EStatGroupingOrSortingMode::ToDescription( *InGroupingMode ) );
 }
 
-FString SFiltersAndPresets::GroupBy_GetSelectedText() const
+FText SFiltersAndPresets::GroupBy_GetSelectedText() const
 {
 	return EStatGroupingOrSortingMode::ToName( GroupingMode );
 }
 
-FString SFiltersAndPresets::GroupBy_GetSelectedTooltipText() const
+FText SFiltersAndPresets::GroupBy_GetSelectedTooltipText() const
 {
 	return EStatGroupingOrSortingMode::ToDescription( GroupingMode );
 }
@@ -1203,7 +1203,7 @@ TSharedRef<SWidget> SFiltersAndPresets::SortBy_OnGenerateWidget( TSharedPtr<ESta
 		.ToolTipText( EStatGroupingOrSortingMode::ToDescription( *InSortingMode ) );
 }
 
-FString SFiltersAndPresets::SortBy_GetSelectedText() const
+FText SFiltersAndPresets::SortBy_GetSelectedText() const
 {
 	return EStatGroupingOrSortingMode::ToName( SortingMode );
 }
diff --git a/Engine/Source/Developer/Profiler/Private/Widgets/SFiltersAndPresets.h b/Engine/Source/Developer/Profiler/Private/Widgets/SFiltersAndPresets.h
--- a/Engine/Source/Developer/Profiler/Private/Widgets/SFiltersAndPresets.h
+++ b/Engine/Source/Developer/Profiler/Private/Widgets/SFiltersAndPresets.h
@@ -41,18 +41,18 @@ struct EStatGroupingOrSortingMode
 	};
 
 	/**
-	 * @param StatGroupingOrSortingMode - The value to get the string for.
+	 * @param StatGroupingOrSortingMode - The value to get the text for.
 	 *
-	 * @return string representation of the specified EStatGroupingOrSortingMode value.
+	 * @return text representation of the specified EStatGroupingOrSortingMode value.
 	 */
-	static FString ToName( const Type StatGroupingOrSortingMode );
+	static FText ToName( const Type StatGroupingOrSortingMode );
 
 	/**
-	 * @param StatGroupingOrSortingMode - The value to get the string for.
+	 * @param StatGroupingOrSortingMode - The value to get the text for.
 	 *
-	 * @return string representation with more detailed explanation of the specified EStatGroupingOrSortingMode value.
+	 * @return text representation with more detailed explanation of the specified EStatGroupingOrSortingMode value.
 	 */
-	static FString ToDescription( const Type StatGroupingOrSortingMode );
+	static FText ToDescription( const Type StatGroupingOrSortingMode );
 
 	/**
 	 * @param StatGroupingOrSortingMode - The value to get the brush name for.
@@ -360,9 +360,9 @@ protected:
 
 	TSharedRef<SWidget> GroupBy_OnGenerateWidget( TSharedPtr<EStatGroupingOrSortingMode::Type> InGroupingMode ) const;
 
-	FString GroupBy_GetSelectedText() const;
+	FText GroupBy_GetSelectedText() const;
 
-	FString GroupBy_GetSelectedTooltipText() const;
+	FText GroupBy_GetSelectedTooltipText() const;
 
 	/*-----------------------------------------------------------------------------
 		SortBy
@@ -372,7 +372,7 @@ protected:
 
 	TSharedRef<SWidget> SortBy_OnGenerateWidget( TSharedPtr<EStatGroupingOrSortingMode::Type> InSortingMode ) const;
 
-	FString SortBy_GetSelectedText() const;
+	FText SortBy_GetSelectedText() const;
 
 protected:
 	/** An array of group and stat nodes generated from the metadata. */
diff --git a/Engine/Source/Developer/Profiler/Private/Widgets/SProfilerSettings.cpp b/Engine/Source/Developer/Profiler/Private/Widgets/SProfilerSettings.cpp
--- a/Engine/Source/Developer/Profiler/Private/Widgets/SProfilerSettings.cpp
+++ b/Engine/Source/Developer/Profiler/Private/Widgets/SProfilerSettings.cpp
@@ -14,13 +14,13 @@ void SProfilerSettings::Construct( const FArguments& InArgs )
 	const TSharedRef<SGridPanel> SettingsGrid = SNew(SGridPanel);
 	int CurrentRowPos = 0;
 
-	AddTitle( LOCTEXT("ProfilerSettingTitle","Profiler settings").ToString(), SettingsGrid, CurrentRowPos );
+	AddTitle( LOCTEXT("ProfilerSettingTitle","Profiler settings"), SettingsGrid, CurrentRowPos );
 	AddSeparator( SettingsGrid, CurrentRowPos );
-	AddHeader( LOCTEXT("MiscTitle","Miscellaneous").ToString(), SettingsGrid, CurrentRowPos );
+	AddHeader( LOCTEXT("MiscTitle","Miscellaneous"), SettingsGrid, CurrentRowPos );
 	AddOption
 	( 
-		LOCTEXT("bSingleInstanceMode_T","Single Instance Mode").ToString(), 
-		LOCTEXT("bSingleInstanceMode_TT","If True, the profiler will work in the single instance mode, all functionality related to multi instances will be disabled or removed from the UI").ToString(), 
+		LOCTEXT("bSingleInstanceMode_T","Single Instance Mode"), 
+		LOCTEXT("bSingleInstanceMode_TT","If True, the profiler will work in the single instance mode, all functionality related to multi instances will be disabled or removed from the UI"), 
 		SettingPtr->bSingleInstanceMode, 
 		SettingPtr->GetDefaults().bSingleInstanceMode,
 		SettingsGrid, 
@@ -28,8 +28,8 @@ void SProfilerSettings::Construct( const FArguments& InArgs )
 	);
 	AddOption
 	( 
-		LOCTEXT("bShowCoalescedViewModesInEventGraph_T","Show Coalesced View Modes In Event Graph").ToString(), 
-		LOCTEXT("bShowCoalescedViewModesInEventGraph_TT","If True, coalesced view modes related functionality will be added to the event graph").ToString(), 
+		LOCTEXT("bShowCoalescedViewModesInEventGraph_T","Show Coalesced View Modes In Event Graph"), 
+		LOCTEXT("bShowCoalescedViewModesInEventGraph_TT","If True, coalesced view modes related functionality will be added to the event graph"), 
 		SettingPtr->bShowCoalescedViewModesInEventGraph, 
 		SettingPtr->GetDefaults().bShowCoalescedViewModesInEventGraph,
 		SettingsGrid, 
@@ -46,7 +46,7 @@ void SProfilerSettings::Construct( const FArguments& InArgs )
 	SettingPtr->EnterEditMode();
 }
 
-void SProfilerSettings::AddTitle( const FString& TitleText, const TSharedRef<SGridPanel>& Grid, int32& RowPos )
+void SProfilerSettings::AddTitle( const FText& TitleText, const TSharedRef<SGridPanel>& Grid, int32& RowPos )
 {
 	Grid->AddSlot( 0, RowPos++ )
 	.Padding( 2.0f )
@@ -70,7 +70,7 @@ void SProfilerSettings::AddSeparator( const TSharedRef<SGridPanel>& Grid, int32&
 	RowPos++;
 }
 
-void SProfilerSettings::AddHeader( const FString& HeaderText, const TSharedRef<SGridPanel>& Grid, int32& RowPos )
+void SProfilerSettings::AddHeader( const FText& HeaderText, const TSharedRef<SGridPanel>& Grid, int32& RowPos )
 {
 	Grid->AddSlot( 0, RowPos++ )
 	.Padding( 2.0f )
@@ -82,7 +82,7 @@ void SProfilerSettings::AddHeader( const FString& HeaderText, const TSharedRef<S
 	RowPos++;
 }
 
-void SProfilerSettings::AddOption( const FString& OptionName, const FString& OptionDesc, bool& Value, const bool& Default, const TSharedRef<SGridPanel>& Grid, int32& RowPos )
+void SProfilerSettings::AddOption(const FText& OptionName, const FText& OptionDesc, bool& Value, const bool& Default, const TSharedRef<SGridPanel>& Grid, int32& RowPos)
 {
 	Grid->AddSlot( 0, RowPos )
 	.Padding( 2.0f )
@@ -117,7 +117,7 @@ void SProfilerSettings::AddOption( const FString& OptionName, const FString& Opt
 		.VAlign( VAlign_Center )
 		[
 			SNew(SButton)
-			.ToolTipText( LOCTEXT("ResetToDefaultToolTip", "Reset to default").ToString() )
+			.ToolTipText( LOCTEXT("ResetToDefaultToolTip", "Reset to default") )
 			.ButtonStyle( FEditorStyle::Get(), TEXT("NoBorder") )
 			.ContentPadding( 0.0f ) 
 			.Visibility( this, &SProfilerSettings::OptionDefault_GetDiffersFromDefaultAsVisibility, (const bool*)&Value, (const bool*)&Default )
@@ -165,7 +165,7 @@ void SProfilerSettings::AddFooter( const TSharedRef<SGridPanel>& Grid, int32& Ro
 				.AutoWidth()
 				[
 					SNew(STextBlock)
-					.Text( LOCTEXT("SaveAndCloseTitle","Save and close").ToString() )
+					.Text( LOCTEXT("SaveAndCloseTitle","Save and close") )
 				]
 			]
 		]
@@ -193,7 +193,7 @@ void SProfilerSettings::AddFooter( const TSharedRef<SGridPanel>& Grid, int32& Ro
 				.AutoWidth()
 				[
 					SNew(STextBlock)
-					.Text( LOCTEXT("ResetToDefaultsTitle","Reset to defaults").ToString() )
+					.Text( LOCTEXT("ResetToDefaultsTitle","Reset to defaults") )
 				]
 			]
 		]
diff --git a/Engine/Source/Developer/Profiler/Private/Widgets/SProfilerSettings.h b/Engine/Source/Developer/Profiler/Private/Widgets/SProfilerSettings.h
--- a/Engine/Source/Developer/Profiler/Private/Widgets/SProfilerSettings.h
+++ b/Engine/Source/Developer/Profiler/Private/Widgets/SProfilerSettings.h
@@ -20,13 +20,13 @@ public:
 	void Construct( const FArguments& InArgs );
 
 private:
-	void AddTitle( const FString& TitleText, const TSharedRef<SGridPanel>& Grid, int32& RowPos );
+	void AddTitle( const FText& TitleText, const TSharedRef<SGridPanel>& Grid, int32& RowPos );
 
 	void AddSeparator( const TSharedRef<SGridPanel>& Grid, int32& RowPos );
 
-	void AddHeader( const FString& HeaderText, const TSharedRef<SGridPanel>& Grid, int32& RowPos);
+	void AddHeader( const FText& HeaderText, const TSharedRef<SGridPanel>& Grid, int32& RowPos);
 
-	void AddOption( const FString& OptionName, const FString& OptionDesc, bool& Value, const bool& Default, const TSharedRef<SGridPanel>& Grid, int32& RowPos );
+	void AddOption( const FText& OptionName, const FText& OptionDesc, bool& Value, const bool& Default, const TSharedRef<SGridPanel>& Grid, int32& RowPos );
 
 	void AddFooter( const TSharedRef<SGridPanel>& Grid, int32& RowPos );
 
diff --git a/Engine/Source/Developer/Profiler/Private/Widgets/SProfilerWindow.cpp b/Engine/Source/Developer/Profiler/Private/Widgets/SProfilerWindow.cpp
--- a/Engine/Source/Developer/Profiler/Private/Widgets/SProfilerWindow.cpp
+++ b/Engine/Source/Developer/Profiler/Private/Widgets/SProfilerWindow.cpp
@@ -250,7 +250,7 @@ void SProfilerWindow::Construct( const FArguments& InArgs, const ISessionManager
 				.Visibility( this, &SProfilerWindow::IsSessionOverlayVissible )
 				[
 					SNew(STextBlock)
-					.Text(LOCTEXT("SelectSessionOverlayText", "Please select a session from the Session Browser or load a saved capture.").ToString())
+					.Text(LOCTEXT("SelectSessionOverlayText", "Please select a session from the Session Browser or load a saved capture."))
 				]
 			]
 
diff --git a/Engine/Source/Developer/SessionFrontend/Private/Widgets/Console/SSessionConsole.cpp b/Engine/Source/Developer/SessionFrontend/Private/Widgets/Console/SSessionConsole.cpp
--- a/Engine/Source/Developer/SessionFrontend/Private/Widgets/Console/SSessionConsole.cpp
+++ b/Engine/Source/Developer/SessionFrontend/Private/Widgets/Console/SSessionConsole.cpp
@@ -147,7 +147,7 @@ void SSessionConsole::Construct( const FArguments& InArgs, ISessionManagerRef In
 					.Visibility(this, &SSessionConsole::HandleSelectSessionOverlayVisibility)
 					[
 						SNew(STextBlock)
-							.Text(LOCTEXT("SelectSessionOverlayText", "Please select a session from the Session Browser").ToString())
+							.Text(LOCTEXT("SelectSessionOverlayText", "Please select a session from the Session Browser"))
 					]
 			]
 	];
diff --git a/Engine/Source/Developer/SessionFrontend/Private/Widgets/Console/SSessionConsoleCommandBar.cpp b/Engine/Source/Developer/SessionFrontend/Private/Widgets/Console/SSessionConsoleCommandBar.cpp
--- a/Engine/Source/Developer/SessionFrontend/Private/Widgets/Console/SSessionConsoleCommandBar.cpp
+++ b/Engine/Source/Developer/SessionFrontend/Private/Widgets/Console/SSessionConsoleCommandBar.cpp
@@ -43,10 +43,10 @@ void SSessionConsoleCommandBar::Construct( const FArguments& InArgs )
 					.ContentPadding(FMargin(6.0f, 2.0f))
 					.IsEnabled(false)
 					.OnClicked(this, &SSessionConsoleCommandBar::HandleSendButtonClicked)
-					.ToolTipText(LOCTEXT("SendButtonTooltip", "Send the command").ToString())
+					.ToolTipText(LOCTEXT("SendButtonTooltip", "Send the command"))
 					[
 						SNew(STextBlock)
-							.Text(LOCTEXT("SendButtonLabel", "Send Command").ToString())
+							.Text(LOCTEXT("SendButtonLabel", "Send Command"))
 					]
 			]
 
@@ -59,10 +59,10 @@ void SSessionConsoleCommandBar::Construct( const FArguments& InArgs )
 					.ContentPadding(FMargin(6.0f, 2.0f))
 					.IsEnabled(false)
 					.OnClicked(this, &SSessionConsoleCommandBar::HandlePromoteToShortcutButtonClicked)
-					.ToolTipText(LOCTEXT("PromoteConsoleCommandButtonTooltip", "Promote Command to Shortcut").ToString())
+					.ToolTipText(LOCTEXT("PromoteConsoleCommandButtonTooltip", "Promote Command to Shortcut"))
 					[
 						SNew(STextBlock)
-							.Text(LOCTEXT("PromoteConsoleCommandButtonLabel", "Promote to Shortcut").ToString())
+							.Text(LOCTEXT("PromoteConsoleCommandButtonLabel", "Promote to Shortcut"))
 					]
 			]
 	];
diff --git a/Engine/Source/Developer/SessionFrontend/Private/Widgets/Console/SSessionConsoleFilterBar.cpp b/Engine/Source/Developer/SessionFrontend/Private/Widgets/Console/SSessionConsoleFilterBar.cpp
--- a/Engine/Source/Developer/SessionFrontend/Private/Widgets/Console/SSessionConsoleFilterBar.cpp
+++ b/Engine/Source/Developer/SessionFrontend/Private/Widgets/Console/SSessionConsoleFilterBar.cpp
@@ -66,7 +66,7 @@ void SSessionConsoleFilterBar::Construct( const FArguments& InArgs )
 					.ButtonContent()
 					[
 						SNew(STextBlock)
-							.Text(LOCTEXT("CategoryComboButtonText", "Categories").ToString())
+							.Text(LOCTEXT("CategoryComboButtonText", "Categories"))
 					]
 					.ContentPadding(FMargin(6.0f, 2.0f))
 					.MenuContent()
diff --git a/Engine/Source/Developer/SessionLauncher/Private/Widgets/Build/SSessionLauncherBuildPage.cpp b/Engine/Source/Developer/SessionLauncher/Private/Widgets/Build/SSessionLauncherBuildPage.cpp
--- a/Engine/Source/Developer/SessionLauncher/Private/Widgets/Build/SSessionLauncherBuildPage.cpp
+++ b/Engine/Source/Developer/SessionLauncher/Private/Widgets/Build/SSessionLauncherBuildPage.cpp
@@ -43,7 +43,7 @@ void SSessionLauncherBuildPage::Construct( const FArguments& InArgs, const FSess
 					.VAlign(VAlign_Center)
 					[
 						SNew(STextBlock)
-							.Text(LOCTEXT("BuildText", "Do you wish to build?").ToString())
+							.Text(LOCTEXT("BuildText", "Do you wish to build?"))
 					]
 
 				+ SHorizontalBox::Slot()
@@ -62,7 +62,7 @@ void SSessionLauncherBuildPage::Construct( const FArguments& InArgs, const FSess
             .Padding(0.0f, 8.0f, 0.0f, 0.0f)
             [
                 SNew(SExpandableArea)
-                .AreaTitle(LOCTEXT("AdvancedAreaTitle", "Advanced Settings").ToString())
+                .AreaTitle(LOCTEXT("AdvancedAreaTitle", "Advanced Settings"))
                 .InitiallyCollapsed(true)
                 .Padding(8.0)
                 .BodyContent()
@@ -73,7 +73,7 @@ void SSessionLauncherBuildPage::Construct( const FArguments& InArgs, const FSess
                     .AutoHeight()
                     [
                         SNew(SButton)
-                        .Text(LOCTEXT("GenDSYMText", "Generate DSYM").ToString())
+                        .Text(LOCTEXT("GenDSYMText", "Generate DSYM"))
                         .IsEnabled( this, &SSessionLauncherBuildPage::HandleGenDSYMButtonEnabled )
                         .OnClicked( this, &SSessionLauncherBuildPage::HandleGenDSYMClicked )
                     ]
diff --git a/Engine/Source/Developer/SessionLauncher/Private/Widgets/Cook/SSessionLauncherCookedPlatforms.h b/Engine/Source/Developer/SessionLauncher/Private/Widgets/Cook/SSessionLauncherCookedPlatforms.h
--- a/Engine/Source/Developer/SessionLauncher/Private/Widgets/Cook/SSessionLauncherCookedPlatforms.h
+++ b/Engine/Source/Developer/SessionLauncher/Private/Widgets/Cook/SSessionLauncherCookedPlatforms.h
@@ -73,8 +73,8 @@ public:
 							// all platforms hyper-link
 							SNew(SHyperlink)
 								.OnNavigate(this, &SSessionLauncherCookedPlatforms::HandleAllPlatformsHyperlinkNavigate, true)
-								.Text(LOCTEXT("AllPlatformsHyperlinkLabel", "All").ToString())
-								.ToolTipText(LOCTEXT("AllPlatformsButtonTooltip", "Select all available platforms.").ToString())
+								.Text(LOCTEXT("AllPlatformsHyperlinkLabel", "All"))
+								.ToolTipText(LOCTEXT("AllPlatformsButtonTooltip", "Select all available platforms."))
 								.Visibility(this, &SSessionLauncherCookedPlatforms::HandleAllPlatformsHyperlinkVisibility)									
 						]
 
@@ -84,8 +84,8 @@ public:
 							// no platforms hyper-link
 							SNew(SHyperlink)
 								.OnNavigate(this, &SSessionLauncherCookedPlatforms::HandleAllPlatformsHyperlinkNavigate, false)
-								.Text(LOCTEXT("NoPlatformsHyperlinkLabel", "None").ToString())
-								.ToolTipText(LOCTEXT("NoPlatformsHyperlinkTooltip", "Deselect all platforms.").ToString())
+								.Text(LOCTEXT("NoPlatformsHyperlinkLabel", "None"))
+								.ToolTipText(LOCTEXT("NoPlatformsHyperlinkTooltip", "Deselect all platforms."))
 								.Visibility(this, &SSessionLauncherCookedPlatforms::HandleAllPlatformsHyperlinkVisibility)									
 						]
 				]
@@ -134,7 +134,7 @@ protected:
 
 		return SNew(STextBlock)
 			.ColorAndOpacity(FLinearColor::Red)
-			.Text(LOCTEXT("NoPlatformsFoundErrorText", "Error: No platforms found.").ToString());
+			.Text(LOCTEXT("NoPlatformsFoundErrorText", "Error: No platforms found."));
 	}
 
 
diff --git a/Engine/Source/Developer/SessionLauncher/Private/Widgets/Deploy/SSessionLauncherDeployFileServerSettings.cpp b/Engine/Source/Developer/SessionLauncher/Private/Widgets/Deploy/SSessionLauncherDeployFileServerSettings.cpp
--- a/Engine/Source/Developer/SessionLauncher/Private/Widgets/Deploy/SSessionLauncherDeployFileServerSettings.cpp
+++ b/Engine/Source/Developer/SessionLauncher/Private/Widgets/Deploy/SSessionLauncherDeployFileServerSettings.cpp
@@ -38,7 +38,7 @@ void SSessionLauncherDeployFileServerSettings::Construct( const FArguments& InAr
 			.Padding(0.0f, 8.0f, 0.0f, 0.0f)
 			[
 				SNew(SExpandableArea)
-					.AreaTitle(LOCTEXT("AdvancedAreaTitle", "Advanced Settings").ToString())
+					.AreaTitle(LOCTEXT("AdvancedAreaTitle", "Advanced Settings"))
 					.InitiallyCollapsed(true)
 					.Padding(8.0f)
 					.BodyContent()
@@ -53,11 +53,11 @@ void SSessionLauncherDeployFileServerSettings::Construct( const FArguments& InAr
 									.IsChecked(this, &SSessionLauncherDeployFileServerSettings::HandleHideWindowCheckBoxIsChecked)
 									.OnCheckStateChanged(this, &SSessionLauncherDeployFileServerSettings::HandleHideWindowCheckBoxCheckStateChanged)
 									.Padding(FMargin(4.0f, 0.0f))
-									.ToolTipText(LOCTEXT("HideWindowCheckBoxTooltip", "If checked, the file server's console window will be hidden from your desktop.").ToString())
+									.ToolTipText(LOCTEXT("HideWindowCheckBoxTooltip", "If checked, the file server's console window will be hidden from your desktop."))
 									.Content()
 									[
 										SNew(STextBlock)
-											.Text(LOCTEXT("HideWindowCheckBoxText", "Hide the file server's console window").ToString())
+											.Text(LOCTEXT("HideWindowCheckBoxText", "Hide the file server's console window"))
 									]
 							]
 
@@ -70,11 +70,11 @@ void SSessionLauncherDeployFileServerSettings::Construct( const FArguments& InAr
 									.IsChecked(this, &SSessionLauncherDeployFileServerSettings::HandleStreamingServerCheckBoxIsChecked)
 									.OnCheckStateChanged(this, &SSessionLauncherDeployFileServerSettings::HandleStreamingServerCheckBoxCheckStateChanged)
 									.Padding(FMargin(4.0f, 0.0f))
-									.ToolTipText(LOCTEXT("StreamingServerCheckBoxTooltip", "If checked, the file server uses an experimental implementation that can serve multiple files simultaneously.").ToString())
+									.ToolTipText(LOCTEXT("StreamingServerCheckBoxTooltip", "If checked, the file server uses an experimental implementation that can serve multiple files simultaneously."))
 									.Content()
 									[
 										SNew(STextBlock)
-											.Text(LOCTEXT("StreamingServerCheckBoxText", "Streaming server (experimental)").ToString())
+											.Text(LOCTEXT("StreamingServerCheckBoxText", "Streaming server (experimental)"))
 									]
 							]
 					]
diff --git a/Engine/Source/Developer/SessionLauncher/Private/Widgets/Deploy/SSessionLauncherDeployRepositorySettings.cpp b/Engine/Source/Developer/SessionLauncher/Private/Widgets/Deploy/SSessionLauncherDeployRepositorySettings.cpp
--- a/Engine/Source/Developer/SessionLauncher/Private/Widgets/Deploy/SSessionLauncherDeployRepositorySettings.cpp
+++ b/Engine/Source/Developer/SessionLauncher/Private/Widgets/Deploy/SSessionLauncherDeployRepositorySettings.cpp
@@ -35,7 +35,7 @@ void SSessionLauncherDeployRepositorySettings::Construct( const FArguments& InAr
 				.AutoHeight()
 				[
 					SNew(STextBlock)
-					.Text(LOCTEXT("RepositoryPathLabel", "Repository Path:").ToString())
+					.Text(LOCTEXT("RepositoryPathLabel", "Repository Path:"))
 				]
 
 				+ SVerticalBox::Slot()
@@ -63,8 +63,8 @@ void SSessionLauncherDeployRepositorySettings::Construct( const FArguments& InAr
 								SNew(SButton)
 								.ContentPadding(FMargin(6.0, 2.0))
 								.IsEnabled(true)
-								.Text(LOCTEXT("BrowseButtonText", "Browse...").ToString())
-								.ToolTipText(LOCTEXT("BrowseButtonToolTip", "Browse for the repository").ToString())
+								.Text(LOCTEXT("BrowseButtonText", "Browse..."))
+								.ToolTipText(LOCTEXT("BrowseButtonToolTip", "Browse for the repository"))
 								.OnClicked(this, &SSessionLauncherDeployRepositorySettings::HandleBrowseButtonClicked)
 							]
 					]
@@ -98,10 +98,9 @@ FReply SSessionLauncherDeployRepositorySettings::HandleBrowseButtonClicked( )
 		void* ParentWindowHandle = (ParentWindow.IsValid() && ParentWindow->GetNativeWindow().IsValid()) ? ParentWindow->GetNativeWindow()->GetOSWindowHandle() : nullptr;
 
 		FString FolderName;
-		const FString Title = LOCTEXT("RepositoryBrowseTitle", "Choose a repository location").ToString();
 		const bool bFolderSelected = DesktopPlatform->OpenDirectoryDialog(
 			ParentWindowHandle,
-			Title,
+			LOCTEXT("RepositoryBrowseTitle", "Choose a repository location").ToString(),
 			RepositoryPathTextBox->GetText().ToString(),
 			FolderName
 			);
diff --git a/Engine/Source/Developer/SessionLauncher/Private/Widgets/Deploy/SSessionLauncherDeployTargets.cpp b/Engine/Source/Developer/SessionLauncher/Private/Widgets/Deploy/SSessionLauncherDeployTargets.cpp
--- a/Engine/Source/Developer/SessionLauncher/Private/Widgets/Deploy/SSessionLauncherDeployTargets.cpp
+++ b/Engine/Source/Developer/SessionLauncher/Private/Widgets/Deploy/SSessionLauncherDeployTargets.cpp
@@ -68,7 +68,7 @@ void SSessionLauncherDeployTargets::Construct( const FArguments& InArgs, const F
 				.VAlign(VAlign_Center)
 				[
 					SNew(STextBlock)
-					.Text(LOCTEXT("SelectDeviceGroupText", "Select or create a new device group to continue.").ToString())
+					.Text(LOCTEXT("SelectDeviceGroupText", "Select or create a new device group to continue."))
 				]
 		]
 
@@ -89,7 +89,7 @@ void SSessionLauncherDeployTargets::Construct( const FArguments& InArgs, const F
 				.VAlign((VAlign_Center))
 				[
 					SNew(STextBlock)
-					.Text(LOCTEXT("ShowLabel", "Show:").ToString())
+					.Text(LOCTEXT("ShowLabel", "Show:"))
 				]
 
 				+ SHorizontalBox::Slot()
@@ -103,7 +103,7 @@ void SSessionLauncherDeployTargets::Construct( const FArguments& InArgs, const F
 						.Style(FEditorStyle::Get(), "RadioButton")
 						[
 							SNew(STextBlock)
-							.Text(LOCTEXT("AllDevicesCheckBoxText", "All available devices").ToString())
+							.Text(LOCTEXT("AllDevicesCheckBoxText", "All available devices"))
 						]
 					]
 
@@ -118,7 +118,7 @@ void SSessionLauncherDeployTargets::Construct( const FArguments& InArgs, const F
 						.Style(FEditorStyle::Get(), "RadioButton")
 						[
 							SNew(STextBlock)
-							.Text(LOCTEXT("GroupDevicesCheckBoxText", "Devices in this group").ToString())
+							.Text(LOCTEXT("GroupDevicesCheckBoxText", "Devices in this group"))
 						]
 					]
 			]
@@ -136,23 +136,23 @@ void SSessionLauncherDeployTargets::Construct( const FArguments& InArgs, const F
 					SNew(SHeaderRow)
 
 					+ SHeaderRow::Column("CheckBox")
-					.DefaultLabel(LOCTEXT("DeviceListCheckboxColumnHeader", " ").ToString())
+					.DefaultLabel(LOCTEXT("DeviceListCheckboxColumnHeader", " "))
 					.FixedWidth(24.0f)
 
 					+ SHeaderRow::Column("Device")
-					.DefaultLabel(LOCTEXT("DeviceListDeviceColumnHeader", "Device").ToString())
+					.DefaultLabel(LOCTEXT("DeviceListDeviceColumnHeader", "Device"))
 					.FillWidth(0.35f)
 
 					+ SHeaderRow::Column("Platform")
-					.DefaultLabel(LOCTEXT("DeviceListPlatformColumnHeader", "Platform").ToString())
+					.DefaultLabel(LOCTEXT("DeviceListPlatformColumnHeader", "Platform"))
 					.FillWidth(0.3f)
 
 					+ SHeaderRow::Column("Host")
-					.DefaultLabel(LOCTEXT("DeviceListHostColumnHeader", "Host").ToString())
+					.DefaultLabel(LOCTEXT("DeviceListHostColumnHeader", "Host"))
 					.FillWidth(0.3f)
 
 					+ SHeaderRow::Column("Owner")
-					.DefaultLabel(LOCTEXT("DeviceListOwnerColumnHeader", "Owner").ToString())
+					.DefaultLabel(LOCTEXT("DeviceListOwnerColumnHeader", "Owner"))
 					.FillWidth(0.35f)
 					)
 					.ListItemsSource(&DeviceProxyList)
@@ -207,23 +207,23 @@ void SSessionLauncherDeployTargets::Construct( const FArguments& InArgs, const F
 						SNew(SHeaderRow)
 
 						+ SHeaderRow::Column("CheckBox")
-						.DefaultLabel(LOCTEXT("DeviceListCheckboxColumnHeader", " ").ToString())
+						.DefaultLabel(LOCTEXT("DeviceListCheckboxColumnHeader", " "))
 						.FixedWidth(24.0f)
 
 						+ SHeaderRow::Column("Device")
-						.DefaultLabel(LOCTEXT("DeviceListDeviceColumnHeader", "Device").ToString())
+						.DefaultLabel(LOCTEXT("DeviceListDeviceColumnHeader", "Device"))
 						.FillWidth(0.35f)
 
 						+ SHeaderRow::Column("Platform")
-						.DefaultLabel(LOCTEXT("DeviceListPlatformColumnHeader", "Platform").ToString())
+						.DefaultLabel(LOCTEXT("DeviceListPlatformColumnHeader", "Platform"))
 						.FillWidth(0.3f)
 
 						+ SHeaderRow::Column("Host")
-						.DefaultLabel(LOCTEXT("DeviceListHostColumnHeader", "Host").ToString())
+						.DefaultLabel(LOCTEXT("DeviceListHostColumnHeader", "Host"))
 						.FillWidth(0.3f)
 
 						+ SHeaderRow::Column("Owner")
-						.DefaultLabel(LOCTEXT("DeviceListOwnerColumnHeader", "Owner").ToString())
+						.DefaultLabel(LOCTEXT("DeviceListOwnerColumnHeader", "Owner"))
 						.FillWidth(0.35f)
 						)
 						.ListItemsSource(&DeviceProxyList)
@@ -354,21 +354,14 @@ bool SSessionLauncherDeployTargets::HandleDeviceListRowIsEnabled( ITargetDeviceP
 }
 
 
-FString SSessionLauncherDeployTargets::HandleDeviceListRowToolTipText( ITargetDeviceProxyPtr DeviceProxy ) const
+FText SSessionLauncherDeployTargets::HandleDeviceListRowToolTipText( ITargetDeviceProxyPtr DeviceProxy ) const
 {
-	FString ToolTipText;
-
-	ToolTipText += LOCTEXT("DeviceListRowToolTipName", "Name: ").ToString() + DeviceProxy->GetName() + LINE_TERMINATOR;
-	ToolTipText += LOCTEXT("DeviceListRowToolTipPlatform", "Platform: ").ToString() + DeviceProxy->GetPlatformName() + LINE_TERMINATOR;
-	ToolTipText += LOCTEXT("DeviceListRowToolTipDeviceId", "Device ID: ").ToString() + DeviceProxy->GetDeviceId();
-/*	ToolTipText += LOCTEXT("DeviceListRowToolTipDeviceType", "Device Type: ").ToString() + DeviceProxy->GetType() + LINE_TERMINATOR;
-	ToolTipText += LOCTEXT("DeviceListRowToolTipMake", "Make: ").ToString() + DeviceProxy->GetMake() + LINE_TERMINATOR;
-	ToolTipText += LOCTEXT("DeviceListRowToolTipModel", "Model: ").ToString() + DeviceProxy->GetModel() + LINE_TERMINATOR;
-	ToolTipText += LOCTEXT("DeviceListRowToolTipHostName", "Host Name: ").ToString() + DeviceProxy->GetHostName() + LINE_TERMINATOR;
-	ToolTipText += LOCTEXT("DeviceListRowToolTipHostUser", "Host User: ").ToString() + DeviceProxy->GetHostUser() + LINE_TERMINATOR;
-	ToolTipText += LOCTEXT("DeviceListRowToolTipStatus", "Connected: ").ToString() + (DeviceProxy->IsConnected() ? *(GYes.ToString()) : *(GNo.ToString()));*/
-
-	return ToolTipText;
+	FTextBuilder Builder;
+	Builder.AppendLineFormat(LOCTEXT("DeviceListRowToolTipName", "Name: {0}"), FText::FromString(DeviceProxy->GetName()));
+	Builder.AppendLineFormat(LOCTEXT("DeviceListRowToolTipPlatform", "Platform: {0}"), FText::FromString(DeviceProxy->GetPlatformName()));
+	Builder.AppendLineFormat(LOCTEXT("DeviceListRowToolTipDeviceId", "Device ID: {0}"), FText::FromString(DeviceProxy->GetDeviceId()));
+
+	return Builder.ToText();
 }
 
 
@@ -429,21 +422,21 @@ EVisibility SSessionLauncherDeployTargets::HandleNoDevicesBoxVisibility( ) const
 }
 
 
-FString SSessionLauncherDeployTargets::HandleNoDevicesTextBlockText( ) const
+FText SSessionLauncherDeployTargets::HandleNoDevicesTextBlockText( ) const
 {
 	if (DeviceProxyList.Num() == 0)
 	{
 		if (ShowDevicesChoice == EShowDevicesChoices::ShowAllDevices)
 		{
-			return LOCTEXT("NoDevicesText", "No available devices were detected.").ToString();
+			return LOCTEXT("NoDevicesText", "No available devices were detected.");
 		}
 		else if (ShowDevicesChoice == EShowDevicesChoices::ShowGroupDevices)
 		{
-			return LOCTEXT("EmptyGroupText", "There are no devices in this group.").ToString();
+			return LOCTEXT("EmptyGroupText", "There are no devices in this group.");
 		}
 	}
 
-	return FString();
+	return FText::GetEmpty();
 }
 
 
diff --git a/Engine/Source/Developer/SessionLauncher/Private/Widgets/Deploy/SSessionLauncherDeployTargets.h b/Engine/Source/Developer/SessionLauncher/Private/Widgets/Deploy/SSessionLauncherDeployTargets.h
--- a/Engine/Source/Developer/SessionLauncher/Private/Widgets/Deploy/SSessionLauncherDeployTargets.h
+++ b/Engine/Source/Developer/SessionLauncher/Private/Widgets/Deploy/SSessionLauncherDeployTargets.h
@@ -84,7 +84,7 @@ private:
 	bool HandleDeviceListRowIsEnabled( ITargetDeviceProxyPtr DeviceProxy ) const;
 
 	// Callback for getting the tool tip text of a device proxy list row.
-	FString HandleDeviceListRowToolTipText( ITargetDeviceProxyPtr DeviceProxy ) const;
+	FText HandleDeviceListRowToolTipText( ITargetDeviceProxyPtr DeviceProxy ) const;
 
 	// Callback for generating a row in the device list view.
 	TSharedRef<ITableRow> HandleDeviceProxyListViewGenerateRow( ITargetDeviceProxyPtr InItem, const TSharedRef<STableViewBase>& OwnerTable ) const;
@@ -102,7 +102,7 @@ private:
 	EVisibility HandleNoDevicesBoxVisibility( ) const;
 
 	// Callback for getting the text in the 'no devices' text block.
-	FString HandleNoDevicesTextBlockText( ) const;
+	FText HandleNoDevicesTextBlockText( ) const;
 
 	// Callback for changing the selected profile in the profile manager.
 	void HandleProfileManagerProfileSelected( const ILauncherProfilePtr& SelectedProfile, const ILauncherProfilePtr& PreviousProfile );
diff --git a/Engine/Source/Developer/SessionLauncher/Private/Widgets/Deploy/SSessionLauncherDeployToDeviceSettings.cpp b/Engine/Source/Developer/SessionLauncher/Private/Widgets/Deploy/SSessionLauncherDeployToDeviceSettings.cpp
--- a/Engine/Source/Developer/SessionLauncher/Private/Widgets/Deploy/SSessionLauncherDeployToDeviceSettings.cpp
+++ b/Engine/Source/Developer/SessionLauncher/Private/Widgets/Deploy/SSessionLauncherDeployToDeviceSettings.cpp
@@ -44,7 +44,7 @@ void SSessionLauncherDeployToDeviceSettings::Construct( const FArguments& InArgs
 					.AutoHeight()
 					[
 						SNew(SExpandableArea)
-							.AreaTitle(LOCTEXT("AdvancedAreaTitle", "Advanced Settings").ToString())
+							.AreaTitle(LOCTEXT("AdvancedAreaTitle", "Advanced Settings"))
 							.InitiallyCollapsed(true)
 							.Padding(8.0f)
 							.BodyContent()
@@ -59,11 +59,11 @@ void SSessionLauncherDeployToDeviceSettings::Construct( const FArguments& InArgs
 											.IsChecked(this, &SSessionLauncherDeployToDeviceSettings::HandleUnrealPakCheckBoxIsChecked)
 											.OnCheckStateChanged(this, &SSessionLauncherDeployToDeviceSettings::HandleUnrealPakCheckBoxCheckStateChanged)
 											.Padding(FMargin(4.0f, 0.0f))
-											.ToolTipText(LOCTEXT("UnrealPakCheckBoxTooltip", "If checked, the content will be deployed as a single UnrealPak file instead of many separate files.").ToString())
+											.ToolTipText(LOCTEXT("UnrealPakCheckBoxTooltip", "If checked, the content will be deployed as a single UnrealPak file instead of many separate files."))
 											.Content()
 											[
 												SNew(STextBlock)
-													.Text(LOCTEXT("UnrealPakCheckBoxText", "Store all content in a single file (UnrealPak)").ToString())
+													.Text(LOCTEXT("UnrealPakCheckBoxText", "Store all content in a single file (UnrealPak)"))
 											]
 									]
 							]
diff --git a/Engine/Source/Developer/SessionLauncher/Private/Widgets/Package/SSessionLauncherPackagingSettings.cpp b/Engine/Source/Developer/SessionLauncher/Private/Widgets/Package/SSessionLauncherPackagingSettings.cpp
--- a/Engine/Source/Developer/SessionLauncher/Private/Widgets/Package/SSessionLauncherPackagingSettings.cpp
+++ b/Engine/Source/Developer/SessionLauncher/Private/Widgets/Package/SSessionLauncherPackagingSettings.cpp
@@ -47,7 +47,7 @@ void SSessionLauncherPackagingSettings::Construct( const FArguments& InArgs, con
 							.AutoHeight()
 							[
 								SNew(STextBlock)
-									.Text(LOCTEXT("RepositoryPathLabel", "Repository Path:").ToString())
+									.Text(LOCTEXT("RepositoryPathLabel", "Repository Path:"))
 							]
 
 						+ SVerticalBox::Slot()
@@ -73,8 +73,8 @@ void SSessionLauncherPackagingSettings::Construct( const FArguments& InArgs, con
 										SNew(SButton)
 											.ContentPadding(FMargin(6.0, 2.0))
 											.IsEnabled(false)
-											.Text(LOCTEXT("BrowseButtonText", "Browse...").ToString())
-											.ToolTipText(LOCTEXT("BrowseButtonToolTip", "Browse for the repository").ToString())
+											.Text(LOCTEXT("BrowseButtonText", "Browse..."))
+											.ToolTipText(LOCTEXT("BrowseButtonToolTip", "Browse for the repository"))
 											//.OnClicked(this, &SSessionLauncherPackagePage::HandleBrowseButtonClicked)
 									]
 							]
@@ -84,7 +84,7 @@ void SSessionLauncherPackagingSettings::Construct( const FArguments& InArgs, con
 							.Padding(0.0, 8.0, 0.0, 0.0)
 							[
 								SNew(STextBlock)
-									.Text(LOCTEXT("DescriptionTextBoxLabel", "Description:").ToString())
+									.Text(LOCTEXT("DescriptionTextBoxLabel", "Description:"))
 							]
 
 						+ SVerticalBox::Slot()
@@ -101,7 +101,7 @@ void SSessionLauncherPackagingSettings::Construct( const FArguments& InArgs, con
 			.Padding(0.0, 8.0, 0.0, 0.0)
 			[
 				SNew(SExpandableArea)
-					.AreaTitle(LOCTEXT("IncrementalPackagingAreaTitle", "Incremental Packaging").ToString())
+					.AreaTitle(LOCTEXT("IncrementalPackagingAreaTitle", "Incremental Packaging"))
 					.InitiallyCollapsed(true)
 					.Padding(8.0)
 					.BodyContent()
@@ -118,7 +118,7 @@ void SSessionLauncherPackagingSettings::Construct( const FArguments& InArgs, con
 									.Content()
 									[
 										SNew(STextBlock)
-											.Text(LOCTEXT("IncrementalCheckBoxText", "This is an incremental build").ToString())
+											.Text(LOCTEXT("IncrementalCheckBoxText", "This is an incremental build"))
 									]
 							]
 
diff --git a/Engine/Source/Developer/SessionLauncher/Private/Widgets/Project/SSessionLauncherProjectPage.cpp b/Engine/Source/Developer/SessionLauncher/Private/Widgets/Project/SSessionLauncherProjectPage.cpp
--- a/Engine/Source/Developer/SessionLauncher/Private/Widgets/Project/SSessionLauncherProjectPage.cpp
+++ b/Engine/Source/Developer/SessionLauncher/Private/Widgets/Project/SSessionLauncherProjectPage.cpp
@@ -31,7 +31,7 @@ void SSessionLauncherProjectPage::Construct( const FArguments& InArgs, const FSe
 					.VAlign(VAlign_Center)
 					[
 						SNew(STextBlock)
-							.Text(LOCTEXT("WhichProjectToUseText", "Which project would you like to use?").ToString())			
+							.Text(LOCTEXT("WhichProjectToUseText", "Which project would you like to use?"))			
 					]
 			]
 
diff --git a/Engine/Source/Developer/SessionLauncher/Private/Widgets/SSessionLauncher.cpp b/Engine/Source/Developer/SessionLauncher/Private/Widgets/SSessionLauncher.cpp
--- a/Engine/Source/Developer/SessionLauncher/Private/Widgets/SSessionLauncher.cpp
+++ b/Engine/Source/Developer/SessionLauncher/Private/Widgets/SSessionLauncher.cpp
@@ -52,219 +52,218 @@ void SSessionLauncher::Construct( const FArguments& InArgs, const TSharedRef<SDo
 		+ SVerticalBox::Slot()
 			.FillHeight(1.0f)
 			[
-				SNew(SHorizontalBox)
+		SNew(SHorizontalBox)
 
-				+ SHorizontalBox::Slot()
-					.AutoWidth()
+		+ SHorizontalBox::Slot()
+			.AutoWidth()
+			[
+				SNew(SBorder)
+					.Padding(FMargin(8.0f, 6.0f, 8.0f, 4.0f))
+					.BorderImage(FEditorStyle::GetBrush("ToolPanel.GroupBorder"))
 					[
-						SNew(SBorder)
-							.Padding(FMargin(8.0f, 6.0f, 8.0f, 4.0f))
-							.BorderImage(FEditorStyle::GetBrush("ToolPanel.GroupBorder"))
-							[
-								// common launcher task bar
-								SAssignNew(Toolbar, SSessionLauncherToolbar, InModel, this)
-							]
+						// common launcher task bar
+						SAssignNew(Toolbar, SSessionLauncherToolbar, InModel, this)
 					]
+			]
 
-				+ SHorizontalBox::Slot()
-					.FillWidth(1.0f)
+		+ SHorizontalBox::Slot()
+			.FillWidth(1.0f)
+			[
+				SNew(SVerticalBox)
+
+				+ SVerticalBox::Slot()
+					.FillHeight(1.0f)
+					.Padding(8.0f, 8.0f, 0.0f, 8.0f)
 					[
-						SNew(SVerticalBox)
+						SAssignNew(WidgetSwitcher, SWidgetSwitcher)
+							.WidgetIndex((int32)ESessionLauncherTasks::QuickLaunch)
 
-						+ SVerticalBox::Slot()
-							.FillHeight(1.0f)
-							.Padding(8.0f, 8.0f, 0.0f, 8.0f)
+						+ SWidgetSwitcher::Slot()
 							[
-								SAssignNew(WidgetSwitcher, SWidgetSwitcher)
-									.WidgetIndex((int32)ESessionLauncherTasks::QuickLaunch)
-
-								+ SWidgetSwitcher::Slot()
-									[
-										// empty panel
-										SNew(SBorder)
-									]
-
-								+ SWidgetSwitcher::Slot()
-									[
-										// deploy build panel
-										SNew(SSessionLauncherBuildTaskSettings, InModel)
-									]
+								// empty panel
+								SNew(SBorder)
+							]
 
-								+ SWidgetSwitcher::Slot()
-									[
-										// deploy build panel
-										SNew(SSessionLauncherDeployTaskSettings, InModel)
-									]
+						+ SWidgetSwitcher::Slot()
+							[
+								// deploy build panel
+								SNew(SSessionLauncherBuildTaskSettings, InModel)
+							]
 
-								+ SWidgetSwitcher::Slot()
-									[
-										// Quick launch panel
-										SNew(SSessionLauncherLaunchTaskSettings, InModel)
-									]
+						+ SWidgetSwitcher::Slot()
+							[
+								// deploy build panel
+								SNew(SSessionLauncherDeployTaskSettings, InModel)
+							]
 
-								+ SWidgetSwitcher::Slot()
-									[
-										// advanced settings panel
-										SNew(SSessionLauncherSettings, InModel)
-									]
+						+ SWidgetSwitcher::Slot()
+							[
+								// Quick launch panel
+								SNew(SSessionLauncherLaunchTaskSettings, InModel)
+							]
 
-								+ SWidgetSwitcher::Slot()
-									[
-										// preview panel
-										SNew(SSessionLauncherPreviewPage, InModel)
-									]
+						+ SWidgetSwitcher::Slot()
+							[
+								// advanced settings panel
+								SNew(SSessionLauncherSettings, InModel)
+							]
 
-								+ SWidgetSwitcher::Slot()
-									.Padding(8.0f, 0.0f)
-									[
-										// progress panel
-										SAssignNew(ProgressPanel, SSessionLauncherProgress)
-									]
+						+ SWidgetSwitcher::Slot()
+							[
+								// preview panel
+								SNew(SSessionLauncherPreviewPage, InModel)
 							]
 
-						+ SVerticalBox::Slot()
-							.AutoHeight()
-							.Padding(8.0f)
+						+ SWidgetSwitcher::Slot()
+							.Padding(8.0f, 0.0f)
 							[
-								SNew(SSessionLauncherValidation, InModel)
+								// progress panel
+								SAssignNew(ProgressPanel, SSessionLauncherProgress)
 							]
+					]
 
-						+ SVerticalBox::Slot()
-							.AutoHeight()
-							.Padding(0.0f, 8.0f, 0.0f, 0.0f)
+				+ SVerticalBox::Slot()
+					.AutoHeight()
+					.Padding(8.0f)
+					[
+						SNew(SSessionLauncherValidation, InModel)
+					]
+
+				+ SVerticalBox::Slot()
+					.AutoHeight()
+					.Padding(0.0f, 8.0f, 0.0f, 0.0f)
+					[
+						SNew(SBorder)
+							.Padding(FMargin(8.0f, 6.0f, 8.0f, 4.0f))
+							.BorderImage(FEditorStyle::GetBrush("ToolPanel.GroupBorder"))
 							[
-								SNew(SBorder)
-									.Padding(FMargin(8.0f, 6.0f, 8.0f, 4.0f))
-									.BorderImage(FEditorStyle::GetBrush("ToolPanel.GroupBorder"))
+								SNew(SHorizontalBox)
+
+								+ SHorizontalBox::Slot()
+									.FillWidth(1.0f)
+									.HAlign(HAlign_Left)
 									[
-										SNew(SHorizontalBox)
+										SNew(SOverlay)
 
-										+ SHorizontalBox::Slot()
-											.FillWidth(1.0f)
-											.HAlign(HAlign_Left)
+										+ SOverlay::Slot()
 											[
-												SNew(SOverlay)
-
-												+ SOverlay::Slot()
+												// edit button
+												SNew(SButton)
+													.ContentPadding(FMargin(6.0f, 2.0f))
+													.OnClicked(this, &SSessionLauncher::HandleEditButtonClicked)
+													.ToolTipText(LOCTEXT("EditButtonTooltip", "Edit the settings of the selected profile."))
 													[
-														// edit button
-														SNew(SButton)
-															.ContentPadding(FMargin(6.0f, 2.0f))
-															.OnClicked(this, &SSessionLauncher::HandleEditButtonClicked)
-															.ToolTipText(LOCTEXT("EditButtonTooltip", "Edit the settings of the selected profile.").ToString())
-															[
-																SNew(STextBlock)
-																	.Text(LOCTEXT("EditButtonLabel", "Edit Settings").ToString())
-															]
-															.Visibility(this, &SSessionLauncher::HandleEditButtonVisibility)
+														SNew(STextBlock)
+															.Text(LOCTEXT("EditButtonLabel", "Edit Settings"))
 													]
+													.Visibility(this, &SSessionLauncher::HandleEditButtonVisibility)
+											]
 
-												+ SOverlay::Slot()
+										+ SOverlay::Slot()
+											[
+												// preview button
+												SNew(SButton)
+													.ContentPadding(FMargin(6.0f, 2.0f))
+													.OnClicked(this, &SSessionLauncher::HandlePreviewButtonClicked)
+													.ToolTipText(LOCTEXT("PreviewButtonTooltip", "Preview the settings of the selected profile."))
 													[
-														// preview button
-														SNew(SButton)
-															.ContentPadding(FMargin(6.0f, 2.0f))
-															.OnClicked(this, &SSessionLauncher::HandlePreviewButtonClicked)
-															.ToolTipText(LOCTEXT("PreviewButtonTooltip", "Preview the settings of the selected profile.").ToString())
-															[
-																SNew(STextBlock)
-																	.Text(LOCTEXT("PreviewButtonLabel", "Preview Settings").ToString())
-															]
-															.Visibility(this, &SSessionLauncher::HandlePreviewButtonVisibility)
+														SNew(STextBlock)
+															.Text(LOCTEXT("PreviewButtonLabel", "Preview Settings"))
 													]
+													.Visibility(this, &SSessionLauncher::HandlePreviewButtonVisibility)
 											]
+									]
+
+								+ SHorizontalBox::Slot()
+									.FillWidth(1.0f)
+									.HAlign(HAlign_Right)
+									[
+										SNew(SOverlay)
 
-										+ SHorizontalBox::Slot()
-											.FillWidth(1.0f)
-											.HAlign(HAlign_Right)
+										+ SOverlay::Slot()
 											[
-												SNew(SOverlay)
+												// launch button
+												SNew(SButton)
+													.ContentPadding(FMargin(6.0f, 2.0f))
+													.IsEnabled(this, &SSessionLauncher::HandleFinishButtonIsEnabled)
+													.OnClicked(this, &SSessionLauncher::HandleFinishButtonClicked)
+													.ToolTipText(LOCTEXT("FinishButtonTooltip", "Launch the task as configured in the selected profile."))
+													[
+														SNew(STextBlock)
+															.Text(LOCTEXT("FinishButtonLabel", "Launch"))
+													]
+													.Visibility(this, &SSessionLauncher::HandleFinishButtonVisibility)
+											]
 
-												+ SOverlay::Slot()
+										+ SOverlay::Slot()
+											[
+												// cancel button
+												SNew(SButton)
+													.ContentPadding(FMargin(6.0f, 2.0f))
+													.OnClicked(this, &SSessionLauncher::HandleCancelButtonClicked)
+													.ToolTipText(LOCTEXT("CancelButtonTooltip", "Cancel the task."))
 													[
-														// launch button
-														SNew(SButton)
-															.ContentPadding(FMargin(6.0f, 2.0f))
-															.IsEnabled(this, &SSessionLauncher::HandleFinishButtonIsEnabled)
-															.OnClicked(this, &SSessionLauncher::HandleFinishButtonClicked)
-															.ToolTipText(LOCTEXT("FinishButtonTooltip", "Launch the task as configured in the selected profile.").ToString())
-															[
-																SNew(STextBlock)
-																	.Text(LOCTEXT("FinishButtonLabel", "Launch").ToString())
-															]
-															.Visibility(this, &SSessionLauncher::HandleFinishButtonVisibility)
+														SNew(STextBlock)
+															.Text(LOCTEXT("CancelButtonLabel", "Cancel"))
 													]
+													.Visibility(this, &SSessionLauncher::HandleCancelButtonVisibility)
+											]
 
-												+ SOverlay::Slot()
+										+ SOverlay::Slot()
+											[
+												SNew(SHorizontalBox)
+													.Visibility(this, &SSessionLauncher::HandleCancelingThrobberVisibility)
+
+												+ SHorizontalBox::Slot()
+													.AutoWidth()
 													[
-														// cancel button
-														SNew(SButton)
-															.ContentPadding(FMargin(6.0f, 2.0f))
-															.OnClicked(this, &SSessionLauncher::HandleCancelButtonClicked)
-															.ToolTipText(LOCTEXT("CancelButtonTooltip", "Cancel the task.").ToString())
-															[
-																SNew(STextBlock)
-																	.Text(LOCTEXT("CancelButtonLabel", "Cancel").ToString())
-															]
-															.Visibility(this, &SSessionLauncher::HandleCancelButtonVisibility)
+														SNew(SCircularThrobber)
+															.NumPieces(5)
+															.Radius(10.0f)
 													]
 
-												+ SOverlay::Slot()
+												+ SHorizontalBox::Slot()
+													.AutoWidth()
+													.Padding(4.0f, 0.0f, 0.0f, 0.0f)
+													.VAlign(VAlign_Center)
 													[
-														SNew(SHorizontalBox)
-															.Visibility(this, &SSessionLauncher::HandleCancelingThrobberVisibility)
+														SNew(STextBlock)
+															.Text(LOCTEXT("CancellingText", "Canceling..."))
+													]
+											]
 
-														+ SHorizontalBox::Slot()
-															.AutoWidth()
-															[
-																SNew(SCircularThrobber)
-																	.NumPieces(5)
-																	.Radius(10.0f)
-															]
+										+ SOverlay::Slot()
+											[
+												SNew(SHorizontalBox)
+													.Visibility(this, &SSessionLauncher::HandleDoneButtonVisibility)
 
-														+ SHorizontalBox::Slot()
-															.AutoWidth()
-															.Padding(4.0f, 0.0f, 0.0f, 0.0f)
-															.VAlign(VAlign_Center)
+												+ SHorizontalBox::Slot()
+													.AutoWidth()
+													.Padding(4.0f, 0.0f, 0.0f, 0.0f)
+													[
+														// re-run button
+														SNew(SButton)
+															.ContentPadding(FMargin(6.0f, 2.0f))
+															.OnClicked(this, &SSessionLauncher::HandleFinishButtonClicked)
+															.ToolTipText(LOCTEXT("RerunButtonTooltip", "Execute the selected profile one more time."))
 															[
 																SNew(STextBlock)
-																	.Text(LOCTEXT("CancellingText", "Canceling...").ToString())
-															]
+																	.Text(LOCTEXT("RerunButtonLabel", "Run Again"))
+															]								
 													]
 
-												+ SOverlay::Slot()
+												+ SHorizontalBox::Slot()
+													.AutoWidth()
+													.Padding(4.0f, 0.0f, 0.0f, 0.0f)
 													[
-														SNew(SHorizontalBox)
-															.Visibility(this, &SSessionLauncher::HandleDoneButtonVisibility)
-
-														+ SHorizontalBox::Slot()
-															.AutoWidth()
-															.Padding(4.0f, 0.0f, 0.0f, 0.0f)
-															[
-																// re-run button
-																SNew(SButton)
-																	.ContentPadding(FMargin(6.0f, 2.0f))
-																	.OnClicked(this, &SSessionLauncher::HandleFinishButtonClicked)
-																	.ToolTipText(LOCTEXT("RerunButtonTooltip", "Execute the selected profile one more time.").ToString())
-																	[
-																		SNew(STextBlock)
-																			.Text(LOCTEXT("RerunButtonLabel", "Run Again").ToString())
-																	]								
-															]
-
-														+ SHorizontalBox::Slot()
-															.AutoWidth()
-															.Padding(4.0f, 0.0f, 0.0f, 0.0f)
+														// done button
+														SNew(SButton)
+															.ContentPadding(FMargin(6.0f, 2.0f))
+															.OnClicked(this, &SSessionLauncher::HandleDoneButtonClicked)
+															.ToolTipText(LOCTEXT("DoneButtonTooltip", "Close this page."))
 															[
-																// done button
-																SNew(SButton)
-																	.ContentPadding(FMargin(6.0f, 2.0f))
-																	.OnClicked(this, &SSessionLauncher::HandleDoneButtonClicked)
-																	.ToolTipText(LOCTEXT("DoneButtonTooltip", "Close this page.").ToString())
-																	[
-																		SNew(STextBlock)
-																			.Text(LOCTEXT("DoneButtonLabel", "Done").ToString())
-																	]
+																SNew(STextBlock)
+																	.Text(LOCTEXT("DoneButtonLabel", "Done"))
 															]
 													]
 											]
@@ -272,6 +271,7 @@ void SSessionLauncher::Construct( const FArguments& InArgs, const TSharedRef<SDo
 							]
 					]
 			]
+			]
 	];
 
 	ConstructUnderMajorTab->SetOnPersistVisualState(SDockTab::FOnPersistVisualState::CreateRaw(this, &SSessionLauncher::HandleMajorTabPersistVisualState));
diff --git a/Engine/Source/Developer/SessionLauncher/Private/Widgets/SSessionLauncherBuildTaskSettings.cpp b/Engine/Source/Developer/SessionLauncher/Private/Widgets/SSessionLauncherBuildTaskSettings.cpp
--- a/Engine/Source/Developer/SessionLauncher/Private/Widgets/SSessionLauncherBuildTaskSettings.cpp
+++ b/Engine/Source/Developer/SessionLauncher/Private/Widgets/SSessionLauncherBuildTaskSettings.cpp
@@ -49,7 +49,7 @@ void SSessionLauncherBuildTaskSettings::Construct( const FArguments& InArgs, con
 							[
 								SNew(STextBlock)
 									.Font(FSlateFontInfo(FPaths::EngineContentDir() / TEXT("Slate/Fonts/Roboto-Bold.ttf"), 13))
-									.Text(LOCTEXT("ProjectSectionHeader", "Project").ToString())
+									.Text(LOCTEXT("ProjectSectionHeader", "Project"))
 							]
 
 						+ SGridPanel::Slot(1, 0)
@@ -73,7 +73,7 @@ void SSessionLauncherBuildTaskSettings::Construct( const FArguments& InArgs, con
 							[
 								SNew(STextBlock)
 									.Font(FSlateFontInfo(FPaths::EngineContentDir() / TEXT("Slate/Fonts/Roboto-Bold.ttf"), 13))
-									.Text(LOCTEXT("DeploySectionHeader", "Deploy").ToString())
+									.Text(LOCTEXT("DeploySectionHeader", "Deploy"))
 							]
 
 						+ SGridPanel::Slot(1, 8)
diff --git a/Engine/Source/Developer/SessionLauncher/Private/Widgets/SSessionLauncherDeployTaskSettings.cpp b/Engine/Source/Developer/SessionLauncher/Private/Widgets/SSessionLauncherDeployTaskSettings.cpp
--- a/Engine/Source/Developer/SessionLauncher/Private/Widgets/SSessionLauncherDeployTaskSettings.cpp
+++ b/Engine/Source/Developer/SessionLauncher/Private/Widgets/SSessionLauncherDeployTaskSettings.cpp
@@ -49,7 +49,7 @@ void SSessionLauncherDeployTaskSettings::Construct( const FArguments& InArgs, co
 							[
 								SNew(STextBlock)
 									.Font(FSlateFontInfo(FPaths::EngineContentDir() / TEXT("Slate/Fonts/Roboto-Bold.ttf"), 13))
-									.Text(LOCTEXT("ProjectSectionHeader", "Project").ToString())
+									.Text(LOCTEXT("ProjectSectionHeader", "Project"))
 							]
 
 						+ SGridPanel::Slot(1, 0)
@@ -73,7 +73,7 @@ void SSessionLauncherDeployTaskSettings::Construct( const FArguments& InArgs, co
 							[
 								SNew(STextBlock)
 									.Font(FSlateFontInfo(FPaths::EngineContentDir() / TEXT("Slate/Fonts/Roboto-Bold.ttf"), 13))
-									.Text(LOCTEXT("DeploySectionHeader", "Deploy").ToString())
+									.Text(LOCTEXT("DeploySectionHeader", "Deploy"))
 							]
 
 						+ SGridPanel::Slot(1, 8)
@@ -97,7 +97,7 @@ void SSessionLauncherDeployTaskSettings::Construct( const FArguments& InArgs, co
 							[
 								SNew(STextBlock)
 									.Font(FSlateFontInfo(FPaths::EngineContentDir() / TEXT("Slate/Fonts/Roboto-Bold.ttf"), 13))
-									.Text(LOCTEXT("LaunchSectionHeader", "Launch").ToString())
+									.Text(LOCTEXT("LaunchSectionHeader", "Launch"))
 							]
 
 						+ SGridPanel::Slot(1, 10)
diff --git a/Engine/Source/Developer/SessionLauncher/Private/Widgets/SSessionLauncherLaunchTaskSettings.cpp b/Engine/Source/Developer/SessionLauncher/Private/Widgets/SSessionLauncherLaunchTaskSettings.cpp
--- a/Engine/Source/Developer/SessionLauncher/Private/Widgets/SSessionLauncherLaunchTaskSettings.cpp
+++ b/Engine/Source/Developer/SessionLauncher/Private/Widgets/SSessionLauncherLaunchTaskSettings.cpp
@@ -49,7 +49,7 @@ void SSessionLauncherLaunchTaskSettings::Construct( const FArguments& InArgs, co
 							[
 								SNew(STextBlock)
 									.Font(FSlateFontInfo(FPaths::EngineContentDir() / TEXT("Slate/Fonts/Roboto-Bold.ttf"), 13))
-									.Text(LOCTEXT("ProjectSectionHeader", "Project").ToString())
+									.Text(LOCTEXT("ProjectSectionHeader", "Project"))
 							]
 
 						+ SGridPanel::Slot(1, 0)
@@ -73,7 +73,7 @@ void SSessionLauncherLaunchTaskSettings::Construct( const FArguments& InArgs, co
 							[
 								SNew(STextBlock)
 								.Font(FSlateFontInfo(FPaths::EngineContentDir() / TEXT("Slate/Fonts/Roboto-Bold.ttf"), 13))
-								.Text(LOCTEXT("CookSectionHeader", "Cook").ToString())
+								.Text(LOCTEXT("CookSectionHeader", "Cook"))
 							]
 
 						+ SGridPanel::Slot(1, 4)
@@ -97,7 +97,7 @@ void SSessionLauncherLaunchTaskSettings::Construct( const FArguments& InArgs, co
 							[
 								SNew(STextBlock)
 									.Font(FSlateFontInfo(FPaths::EngineContentDir() / TEXT("Slate/Fonts/Roboto-Bold.ttf"), 13))
-									.Text(LOCTEXT("TargetSectionHeader", "Target").ToString())
+									.Text(LOCTEXT("TargetSectionHeader", "Target"))
 							]
 
 						+ SGridPanel::Slot(1, 10)
