Fixed up redundant FText .ToString() usage and some improper FText usage.

diff --git a/Engine/Source/Developer/SessionLauncher/Private/Widgets/SSessionLauncherPreviewPage.cpp b/Engine/Source/Developer/SessionLauncher/Private/Widgets/SSessionLauncherPreviewPage.cpp
--- a/Engine/Source/Developer/SessionLauncher/Private/Widgets/SSessionLauncherPreviewPage.cpp
+++ b/Engine/Source/Developer/SessionLauncher/Private/Widgets/SSessionLauncherPreviewPage.cpp
@@ -35,7 +35,7 @@ void SSessionLauncherPreviewPage::Construct( const FArguments& InArgs, const FSe
 					[
 						SNew(STextBlock)
 							.Font(FSlateFontInfo(FPaths::EngineContentDir() / TEXT("Slate/Fonts/Roboto-Bold.ttf"), 13))
-							.Text(LOCTEXT("BuildSectionHeader", "Build").ToString())
+							.Text(LOCTEXT("BuildSectionHeader", "Build"))
 					]
 
 				+ SGridPanel::Slot(1, 0)
@@ -58,7 +58,7 @@ void SSessionLauncherPreviewPage::Construct( const FArguments& InArgs, const FSe
 											.AutoWidth()
 											[
 												SNew(STextBlock)
-													.Text(LOCTEXT("ProjectLabel", "Project:").ToString())
+													.Text(LOCTEXT("ProjectLabel", "Project:"))
 											]
 
 										+ SHorizontalBox::Slot()
@@ -90,7 +90,7 @@ void SSessionLauncherPreviewPage::Construct( const FArguments& InArgs, const FSe
 											.AutoWidth()
 											[
 												SNew(STextBlock)
-													.Text(LOCTEXT("ConfigurationLabel", "Build Configuration:").ToString())
+													.Text(LOCTEXT("ConfigurationLabel", "Build Configuration:"))
 											]
 
 										+ SHorizontalBox::Slot()
@@ -122,7 +122,7 @@ void SSessionLauncherPreviewPage::Construct( const FArguments& InArgs, const FSe
 											.AutoWidth()
 											[
 												SNew(STextBlock)
-													.Text(LOCTEXT("PlatformsLabel", "Platforms:").ToString())
+													.Text(LOCTEXT("PlatformsLabel", "Platforms:"))
 											]
 
 										+ SHorizontalBox::Slot()
@@ -161,7 +161,7 @@ void SSessionLauncherPreviewPage::Construct( const FArguments& InArgs, const FSe
 					[
 						SNew(STextBlock)
 							.Font(FSlateFontInfo(FPaths::EngineContentDir() / TEXT("Slate/Fonts/Roboto-Bold.ttf"), 13))
-							.Text(LOCTEXT("CookSectionHeader", "Cook").ToString())
+							.Text(LOCTEXT("CookSectionHeader", "Cook"))
 					]
 
 				+ SGridPanel::Slot(1, 2)
@@ -184,7 +184,7 @@ void SSessionLauncherPreviewPage::Construct( const FArguments& InArgs, const FSe
 											.AutoWidth()
 											[
 												SNew(STextBlock)
-													.Text(LOCTEXT("CookModeLabel", "Cook Mode:").ToString())
+													.Text(LOCTEXT("CookModeLabel", "Cook Mode:"))
 											]
 
 										+ SHorizontalBox::Slot()
@@ -194,7 +194,7 @@ void SSessionLauncherPreviewPage::Construct( const FArguments& InArgs, const FSe
 											[
 												// do not cook
 												SNew(STextBlock)
-													.Text(LOCTEXT("DoNotCookLabel", "Do not cook").ToString())
+													.Text(LOCTEXT("DoNotCookLabel", "Do not cook"))
 											]
 									]
 							]
@@ -214,7 +214,7 @@ void SSessionLauncherPreviewPage::Construct( const FArguments& InArgs, const FSe
 											.AutoWidth()
 											[
 												SNew(STextBlock)
-													.Text(LOCTEXT("BuildModeLabel", "Cook Mode:").ToString())
+													.Text(LOCTEXT("BuildModeLabel", "Cook Mode:"))
 											]
 
 										+ SHorizontalBox::Slot()
@@ -224,7 +224,7 @@ void SSessionLauncherPreviewPage::Construct( const FArguments& InArgs, const FSe
 											[
 												// cook on the fly
 												SNew(STextBlock)
-													.Text(LOCTEXT("CookOnTheFlyLabel", "On the fly").ToString())
+													.Text(LOCTEXT("CookOnTheFlyLabel", "On the fly"))
 											]
 									]
 
@@ -238,7 +238,7 @@ void SSessionLauncherPreviewPage::Construct( const FArguments& InArgs, const FSe
 											.AutoWidth()
 											[
 												SNew(STextBlock)
-													.Text(LOCTEXT("CookerOptionsLabel", "Cooker Options:").ToString())
+													.Text(LOCTEXT("CookerOptionsLabel", "Cooker Options:"))
 											]
 
 										+ SHorizontalBox::Slot()
@@ -268,7 +268,7 @@ void SSessionLauncherPreviewPage::Construct( const FArguments& InArgs, const FSe
 											.AutoWidth()
 											[
 												SNew(STextBlock)
-													.Text(LOCTEXT("BuildModeLabel", "Cook Mode:").ToString())
+													.Text(LOCTEXT("BuildModeLabel", "Cook Mode:"))
 											]
 
 										+ SHorizontalBox::Slot()
@@ -278,7 +278,7 @@ void SSessionLauncherPreviewPage::Construct( const FArguments& InArgs, const FSe
 											[
 												// cook by the book
 												SNew(STextBlock)
-													.Text(LOCTEXT("CookByTheBookLabel", "By the book").ToString())
+													.Text(LOCTEXT("CookByTheBookLabel", "By the book"))
 											]
 									]
 
@@ -292,7 +292,7 @@ void SSessionLauncherPreviewPage::Construct( const FArguments& InArgs, const FSe
 											.AutoWidth()
 											[
 												SNew(STextBlock)
-													.Text(LOCTEXT("CulturesBuildLabel", "Cooked Cultures:").ToString())
+													.Text(LOCTEXT("CulturesBuildLabel", "Cooked Cultures:"))
 											]
 
 										+ SHorizontalBox::Slot()
@@ -324,7 +324,7 @@ void SSessionLauncherPreviewPage::Construct( const FArguments& InArgs, const FSe
 											.AutoWidth()
 											[
 												SNew(STextBlock)
-													.Text(LOCTEXT("MapsBuildLabel", "Cooked Maps:").ToString())
+													.Text(LOCTEXT("MapsBuildLabel", "Cooked Maps:"))
 											]
 
 										+ SHorizontalBox::Slot()
@@ -348,7 +348,7 @@ void SSessionLauncherPreviewPage::Construct( const FArguments& InArgs, const FSe
 											.AutoWidth()
 											[
 												SNew(STextBlock)
-													.Text(LOCTEXT("CookerOptionsLabel", "Cooker Options:").ToString())
+													.Text(LOCTEXT("CookerOptionsLabel", "Cooker Options:"))
 											]
 
 										+ SHorizontalBox::Slot()
@@ -380,7 +380,7 @@ void SSessionLauncherPreviewPage::Construct( const FArguments& InArgs, const FSe
 					[
 						SNew(STextBlock)
 							.Font(FSlateFontInfo(FPaths::EngineContentDir() / TEXT("Slate/Fonts/Roboto-Bold.ttf"), 13))
-							.Text(LOCTEXT("PackageSectionHeader", "Package").ToString())
+							.Text(LOCTEXT("PackageSectionHeader", "Package"))
 					]
 
 				+ SGridPanel::Slot(1, 4)
@@ -405,7 +405,7 @@ void SSessionLauncherPreviewPage::Construct( const FArguments& InArgs, const FSe
 					[
 						SNew(STextBlock)
 							.Font(FSlateFontInfo(FPaths::EngineContentDir() / TEXT("Slate/Fonts/Roboto-Bold.ttf"), 13))
-							.Text(LOCTEXT("DeploySectionHeader", "Deploy").ToString())
+							.Text(LOCTEXT("DeploySectionHeader", "Deploy"))
 					]
 
 				+ SGridPanel::Slot(1, 6)
@@ -428,7 +428,7 @@ void SSessionLauncherPreviewPage::Construct( const FArguments& InArgs, const FSe
 											.AutoWidth()
 											[
 												SNew(STextBlock)
-													.Text(LOCTEXT("DeployModeLabel", "Deploy Mode:").ToString())
+													.Text(LOCTEXT("DeployModeLabel", "Deploy Mode:"))
 											]
 
 										+ SHorizontalBox::Slot()
@@ -438,7 +438,7 @@ void SSessionLauncherPreviewPage::Construct( const FArguments& InArgs, const FSe
 											[
 												// do not cook
 												SNew(STextBlock)
-													.Text(LOCTEXT("DoNotDeployLabel", "Do not deploy").ToString())
+													.Text(LOCTEXT("DoNotDeployLabel", "Do not deploy"))
 											]
 									]
 							]
@@ -556,7 +556,7 @@ void SSessionLauncherPreviewPage::Construct( const FArguments& InArgs, const FSe
 																	.Visibility(EVisibility::Collapsed)
 
 																+ SHeaderRow::Column("Device")
-																	.DefaultLabel(LOCTEXT("DeviceListDeviceColumnHeader", "Device").ToString())
+																	.DefaultLabel(LOCTEXT("DeviceListDeviceColumnHeader", "Device"))
 															)
 													]
 											]
@@ -837,50 +837,46 @@ void SSessionLauncherPreviewPage::RefreshDeviceProxyList( )
 /* SSessionLauncherPreviewPage callbacks
  *****************************************************************************/
 
-FString SSessionLauncherPreviewPage::HandleBuildConfigurationTextBlockText( ) const
+FText SSessionLauncherPreviewPage::HandleBuildConfigurationTextBlockText( ) const
 {
 	ILauncherProfilePtr SelectedProfile = Model->GetSelectedProfile();
 
 	if (SelectedProfile.IsValid())
 	{
-		return EBuildConfigurations::ToString(SelectedProfile->GetBuildConfiguration());
+		return FText::FromString(EBuildConfigurations::ToString(SelectedProfile->GetBuildConfiguration()));
 	}
 
-	return LOCTEXT("NotAvailableText", "n/a").ToString();
+	return LOCTEXT("NotAvailableText", "n/a");
 }
 
 
-FString SSessionLauncherPreviewPage::HandleBuildPlatformsTextBlockText( ) const
+FText SSessionLauncherPreviewPage::HandleBuildPlatformsTextBlockText( ) const
 {
-	FString Result;
-
 	ILauncherProfilePtr SelectedProfile = Model->GetSelectedProfile();
 
 	if (SelectedProfile.IsValid())
 	{
 		const TArray<FString>& Platforms = SelectedProfile->GetCookedPlatforms();
-
-		for (int32 PlatformIndex = 0; PlatformIndex < Platforms.Num(); ++PlatformIndex)
+		if (Platforms.Num() > 0)
 		{
-			Result += Platforms[PlatformIndex];
-
-			if (PlatformIndex + 1 < Platforms.Num())
+			FTextBuilder Builder;
+			for (int32 PlatformIndex = 0; PlatformIndex < Platforms.Num(); ++PlatformIndex)
 			{
-				Result += LINE_TERMINATOR;
+				Builder.AppendLine(Platforms[PlatformIndex]);
 			}
+			return Builder.ToText();
 		}
-
-		if (Result.IsEmpty())
+		else
 		{
-			Result = LOCTEXT("NotSetText", "<not set>").ToString();
+			return LOCTEXT("NotSetText", "<not set>");
 		}
 	}
 
-	return Result;
+	return FText::GetEmpty();
 }
 
 
-FString SSessionLauncherPreviewPage::HandleCommandLineTextBlockText( ) const
+FText SSessionLauncherPreviewPage::HandleCommandLineTextBlockText( ) const
 {
 	ILauncherProfilePtr SelectedProfile = Model->GetSelectedProfile();
 
@@ -890,93 +886,83 @@ FString SSessionLauncherPreviewPage::HandleCommandLineTextBlockText( ) const
 
 		if (CommandLine.IsEmpty())
 		{
-			return LOCTEXT("EmptyText", "<empty>").ToString();
+			return LOCTEXT("EmptyText", "<empty>");
 		}
 
-		return CommandLine;
+		return FText::FromString( CommandLine );
 	}
 
-	return LOCTEXT("NotAvailableText", "n/a").ToString();
+	return LOCTEXT("NotAvailableText", "n/a");
 }
 
 
-FString SSessionLauncherPreviewPage::HandleCookedCulturesTextBlockText( ) const
+FText SSessionLauncherPreviewPage::HandleCookedCulturesTextBlockText( ) const
 {
-	FString Result;
-
 	ILauncherProfilePtr SelectedProfile = Model->GetSelectedProfile();
 
 	if (SelectedProfile.IsValid())
 	{
 		const TArray<FString>& CookedCultures = SelectedProfile->GetCookedCultures();
-
-		for (int32 CookedCultureIndex = 0; CookedCultureIndex < CookedCultures.Num(); ++CookedCultureIndex)
+		if (CookedCultures.Num() > 0)
 		{
-			Result += CookedCultures[CookedCultureIndex];
-			
-			if (CookedCultureIndex + 1 < CookedCultures.Num())
+			FTextBuilder Builder;
+			for (int32 CookedCultureIndex = 0; CookedCultureIndex < CookedCultures.Num(); ++CookedCultureIndex)
 			{
-				Result += LINE_TERMINATOR;
+				Builder.AppendLine(CookedCultures[CookedCultureIndex]);
 			}
+			return Builder.ToText();
 		}
-
-		if (Result.IsEmpty())
+		else
 		{
-			Result = LOCTEXT("NotSetText", "<not set>").ToString();
+			return LOCTEXT("NotSetText", "<not set>");
 		}
 	}
 
-	return Result;
+	return FText::GetEmpty();
 }
 
 
-FString SSessionLauncherPreviewPage::HandleCookedMapsTextBlockText( ) const
+FText SSessionLauncherPreviewPage::HandleCookedMapsTextBlockText( ) const
 {
-	FString Result;
-
 	ILauncherProfilePtr SelectedProfile = Model->GetSelectedProfile();
 
 	if (SelectedProfile.IsValid())
 	{
 		const TArray<FString>& CookedMaps = SelectedProfile->GetCookedMaps();
-
-		for (int32 CookedMapIndex = 0; CookedMapIndex < CookedMaps.Num(); ++CookedMapIndex)
+		if (CookedMaps.Num() > 0)
 		{
-			Result += CookedMaps[CookedMapIndex];
-
-			if (CookedMapIndex + 1 < CookedMaps.Num())
+			FTextBuilder Builder;
+			for (int32 CookedMapIndex = 0; CookedMapIndex < CookedMaps.Num(); ++CookedMapIndex)
 			{
-				Result += LINE_TERMINATOR;
+				Builder.AppendLine(CookedMaps[CookedMapIndex]);
 			}
+			return Builder.ToText();
 		}
-
-		if (Result.IsEmpty())
+		else
 		{
-			Result = LOCTEXT("NoneText", "<none>").ToString();
+			return LOCTEXT("NotSetText", "<not set>");
 		}
 	}
 
-	return Result;
+	return FText::GetEmpty();
 }
 
 
-FString SSessionLauncherPreviewPage::HandleCookerOptionsTextBlockText( ) const
+FText SSessionLauncherPreviewPage::HandleCookerOptionsTextBlockText( ) const
 {
-	FString Result;
-
 	ILauncherProfilePtr SelectedProfile = Model->GetSelectedProfile();
 
 	if (SelectedProfile.IsValid())
 	{
 		if (SelectedProfile->GetCookOptions().IsEmpty())
 		{
-			return LOCTEXT("NoneText", "<none>").ToString();
+			return LOCTEXT("NoneText", "<none>");
 		}
 
-		return SelectedProfile->GetCookOptions();
+		return FText::FromString(SelectedProfile->GetCookOptions());
 	}
 
-	return Result;
+	return FText::GetEmpty();
 }
 
 
diff --git a/Engine/Source/Developer/SessionLauncher/Private/Widgets/SSessionLauncherPreviewPage.h b/Engine/Source/Developer/SessionLauncher/Private/Widgets/SSessionLauncherPreviewPage.h
--- a/Engine/Source/Developer/SessionLauncher/Private/Widgets/SSessionLauncherPreviewPage.h
+++ b/Engine/Source/Developer/SessionLauncher/Private/Widgets/SSessionLauncherPreviewPage.h
@@ -46,22 +46,22 @@ protected:
 private:
 
 	// Callback for getting the text in the 'Build Configuration' text block.
-	FString HandleBuildConfigurationTextBlockText( ) const;
+	FText HandleBuildConfigurationTextBlockText( ) const;
 
 	// Callback for getting the list of platforms to build for.
-	FString HandleBuildPlatformsTextBlockText( ) const;
+	FText HandleBuildPlatformsTextBlockText() const;
 
 	// Callback for getting the text in the 'Command Line' text block.
-	FString HandleCommandLineTextBlockText( ) const;
+	FText HandleCommandLineTextBlockText() const;
 
 	// Callback for getting the text in the 'Cooked Cultures' text block.
-	FString HandleCookedCulturesTextBlockText( ) const;
+	FText HandleCookedCulturesTextBlockText() const;
 
 	// Callback for getting the text in the 'Cooked Maps' text block.
-	FString HandleCookedMapsTextBlockText( ) const;
+	FText HandleCookedMapsTextBlockText() const;
 
 	// Callback for getting the text in the 'Cooker Options' text block.
-	FString HandleCookerOptionsTextBlockText( ) const;
+	FText HandleCookerOptionsTextBlockText() const;
 
 	// Callback for getting the visibility of the specified cook summary box.
 	EVisibility HandleCookSummaryBoxVisibility( ELauncherProfileCookModes::Type CookMode ) const;
diff --git a/Engine/Source/Developer/SessionLauncher/Private/Widgets/SSessionLauncherProfileSelector.h b/Engine/Source/Developer/SessionLauncher/Private/Widgets/SSessionLauncherProfileSelector.h
--- a/Engine/Source/Developer/SessionLauncher/Private/Widgets/SSessionLauncherProfileSelector.h
+++ b/Engine/Source/Developer/SessionLauncher/Private/Widgets/SSessionLauncherProfileSelector.h
@@ -58,7 +58,7 @@ public:
 				.VAlign(VAlign_Center)
 				[
 					SNew(STextBlock)
-						.Text(LOCTEXT("ProfileComboBoxLabel", "Profile:").ToString())
+						.Text(LOCTEXT("ProfileComboBoxLabel", "Profile:"))
 				]
 
 			+ SHorizontalBox::Slot()
@@ -69,8 +69,8 @@ public:
 					SAssignNew(ProfileComboBox, SEditableComboBox<ILauncherProfilePtr>)
 						.InitiallySelectedItem(Model->GetSelectedProfile())
 						.OptionsSource(&Model->GetProfileManager()->GetAllProfiles())
-						.AddButtonToolTip(LOCTEXT("AddProfileButtonToolTip", "Add a new profile").ToString())
-						.RemoveButtonToolTip(LOCTEXT("DeleteProfileButtonToolTip", "Delete the selected profile").ToString())
+						.AddButtonToolTip(LOCTEXT("AddProfileButtonToolTip", "Add a new profile"))
+						.RemoveButtonToolTip(LOCTEXT("DeleteProfileButtonToolTip", "Delete the selected profile"))
 						.OnAddClicked(this, &SSessionLauncherProfileSelector::HandleProfileComboBoxAddClicked)
 						.OnGenerateWidget(this, &SSessionLauncherProfileSelector::HandleProfileComboBoxGenerateWidget)
 						.OnGetEditableText(this, &SSessionLauncherProfileSelector::HandleProfileComboBoxGetEditableText)
diff --git a/Engine/Source/Developer/SessionLauncher/Private/Widgets/SSessionLauncherProgress.h b/Engine/Source/Developer/SessionLauncher/Private/Widgets/SSessionLauncherProgress.h
--- a/Engine/Source/Developer/SessionLauncher/Private/Widgets/SSessionLauncherProgress.h
+++ b/Engine/Source/Developer/SessionLauncher/Private/Widgets/SSessionLauncherProgress.h
@@ -82,19 +82,19 @@ public:
 									SNew(SHeaderRow)
 
 									+ SHeaderRow::Column("Icon")
-										.DefaultLabel(LOCTEXT("TaskListIconColumnHeader", " ").ToString())
+										.DefaultLabel(LOCTEXT("TaskListIconColumnHeader", " "))
 										.FixedWidth(20.0)
 
 									+ SHeaderRow::Column("Task")
-										.DefaultLabel(LOCTEXT("TaskListTaskColumnHeader", "Task").ToString())
+										.DefaultLabel(LOCTEXT("TaskListTaskColumnHeader", "Task"))
 										.FillWidth(1.0)
 
 									+ SHeaderRow::Column("Duration")
-										.DefaultLabel(LOCTEXT("TaskListDurationColumnHeader", "Duration").ToString())
+										.DefaultLabel(LOCTEXT("TaskListDurationColumnHeader", "Duration"))
 										.FixedWidth(64.0)
 
 									+ SHeaderRow::Column("Status")
-										.DefaultLabel(LOCTEXT("TaskListStatusColumnHeader", "Status").ToString())
+										.DefaultLabel(LOCTEXT("TaskListStatusColumnHeader", "Status"))
 										.FixedWidth(80.0)
 								)
 								.ListItemsSource(&TaskList)
@@ -119,7 +119,7 @@ public:
 							SNew(SHeaderRow)
 
 							+ SHeaderRow::Column("Status")
-							.DefaultLabel(LOCTEXT("TaskListStatusColumnHeader", "Status").ToString())
+							.DefaultLabel(LOCTEXT("TaskListStatusColumnHeader", "Status"))
 							.FillWidth(1.0)
 						)
 						.ListItemsSource(&MessageList)
@@ -142,8 +142,8 @@ public:
 						SAssignNew(CopyButton, SButton)
 						.ContentPadding(FMargin(6.0f, 2.0f))
 						.IsEnabled(false)
-						.Text(LOCTEXT("CopyButtonText", "Copy").ToString())
-						.ToolTipText(LOCTEXT("CopyButtonTooltip", "Copy the selected log messages to the clipboard").ToString())
+						.Text(LOCTEXT("CopyButtonText", "Copy"))
+						.ToolTipText(LOCTEXT("CopyButtonTooltip", "Copy the selected log messages to the clipboard"))
 						.OnClicked(this, &SSessionLauncherProgress::HandleCopyButtonClicked)
 					]
 
@@ -155,8 +155,8 @@ public:
 							SAssignNew(ClearButton, SButton)
 							.ContentPadding(FMargin(6.0f, 2.0f))
 							.IsEnabled(false)
-							.Text(LOCTEXT("ClearButtonText", "Clear Log").ToString())
-							.ToolTipText(LOCTEXT("ClearButtonTooltip", "Clear the log window").ToString())
+							.Text(LOCTEXT("ClearButtonText", "Clear Log"))
+							.ToolTipText(LOCTEXT("ClearButtonTooltip", "Clear the log window"))
 							.OnClicked(this, &SSessionLauncherProgress::HandleClearButtonClicked)
 						]
 
@@ -168,8 +168,8 @@ public:
 							SAssignNew(SaveButton, SButton)
 							.ContentPadding(FMargin(6.0f, 2.0f))
 							.IsEnabled(false)
-							.Text(LOCTEXT("ExportButtonText", "Save Log...").ToString())
-							.ToolTipText(LOCTEXT("SaveButtonTooltip", "Save the entire log to a file").ToString())
+							.Text(LOCTEXT("ExportButtonText", "Save Log..."))
+							.ToolTipText(LOCTEXT("SaveButtonTooltip", "Save the entire log to a file"))
 							.Visibility((FDesktopPlatformModule::Get() != NULL) ? EVisibility::Visible : EVisibility::Collapsed)
 							.OnClicked(this, &SSessionLauncherProgress::HandleSaveButtonClicked)
 						]
@@ -257,7 +257,7 @@ private:
 	}
 
 	// Callback for getting the text of the progress text box.
-	FString HandleProgressTextBlockText( ) const
+	FText HandleProgressTextBlockText( ) const
 	{
 		ILauncherWorkerPtr LauncherWorkerPtr = LauncherWorker.Pin();
 
@@ -266,7 +266,7 @@ private:
 			if ((LauncherWorkerPtr->GetStatus() == ELauncherWorkerStatus::Busy) ||
 				(LauncherWorkerPtr->GetStatus() == ELauncherWorkerStatus::Canceling))
 			{
-				return LOCTEXT("OperationInProgressText", "Operation in progress...").ToString();
+				return LOCTEXT("OperationInProgressText", "Operation in progress...");
 			}
 
 			int32 NumCanceled = 0;
@@ -291,10 +291,15 @@ private:
 				}
 			}
 
-			return FString::Printf(*LOCTEXT("TasksFinishedFormatText", "Operation finished. Completed: %i, Failed: %i, Canceled: %i").ToString(), NumCompleted, NumFailed, NumCanceled);
+			FFormatNamedArguments Args;
+			Args.Add(TEXT("NumCompleted"), NumCompleted);
+			Args.Add(TEXT("NumFailed"), NumCompleted);
+			Args.Add(TEXT("NumCanceled"), NumCanceled);
+
+			return FText::Format(LOCTEXT("TasksFinishedFormatText", "Operation finished. Completed: {NumCompleted}, Failed: {NumFailed}, Canceled: {NumCanceled}"), Args);
 		}
 
-		return FString();
+		return FText::GetEmpty();
 	}
 
 	// Callback for generating a row in the task list view.
diff --git a/Engine/Source/Developer/SessionLauncher/Private/Widgets/SSessionLauncherSettings.cpp b/Engine/Source/Developer/SessionLauncher/Private/Widgets/SSessionLauncherSettings.cpp
--- a/Engine/Source/Developer/SessionLauncher/Private/Widgets/SSessionLauncherSettings.cpp
+++ b/Engine/Source/Developer/SessionLauncher/Private/Widgets/SSessionLauncherSettings.cpp
@@ -68,7 +68,7 @@ void SSessionLauncherSettings::Construct( const FArguments& InArgs, const FSessi
 						[
 							SNew(STextBlock)
 							.Font(FSlateFontInfo(FPaths::EngineContentDir() / TEXT("Slate/Fonts/Roboto-Bold.ttf"), 13))
-							.Text(LOCTEXT("ProjectSectionHeader", "Project").ToString())
+							.Text(LOCTEXT("ProjectSectionHeader", "Project"))
 						]
 
 						+ SGridPanel::Slot(1, 0)
@@ -92,7 +92,7 @@ void SSessionLauncherSettings::Construct( const FArguments& InArgs, const FSessi
 							[
 								SNew(STextBlock)
 								.Font(FSlateFontInfo(FPaths::EngineContentDir() / TEXT("Slate/Fonts/Roboto-Bold.ttf"), 13))
-								.Text(LOCTEXT("BuildSectionHeader", "Build").ToString())
+								.Text(LOCTEXT("BuildSectionHeader", "Build"))
 							]
 
 						+ SGridPanel::Slot(1, 2)
@@ -116,7 +116,7 @@ void SSessionLauncherSettings::Construct( const FArguments& InArgs, const FSessi
 							[
 								SNew(STextBlock)
 								.Font(FSlateFontInfo(FPaths::EngineContentDir() / TEXT("Slate/Fonts/Roboto-Bold.ttf"), 13))
-								.Text(LOCTEXT("CookSectionHeader", "Cook").ToString())
+								.Text(LOCTEXT("CookSectionHeader", "Cook"))
 							]
 
 						+ SGridPanel::Slot(1, 4)
@@ -139,7 +139,7 @@ void SSessionLauncherSettings::Construct( const FArguments& InArgs, const FSessi
 							[
 								SNew(STextBlock)
 								.Font(FSlateFontInfo(FPaths::EngineContentDir() / TEXT("Slate/Fonts/Roboto-Bold.ttf"), 13))
-								.Text(LOCTEXT("PackageSectionHeader", "Package").ToString())
+								.Text(LOCTEXT("PackageSectionHeader", "Package"))
 							]
 
 						+ SGridPanel::Slot(1, 6)
@@ -163,7 +163,7 @@ void SSessionLauncherSettings::Construct( const FArguments& InArgs, const FSessi
 							[
 								SNew(STextBlock)
 								.Font(FSlateFontInfo(FPaths::EngineContentDir() / TEXT("Slate/Fonts/Roboto-Bold.ttf"), 13))
-								.Text(LOCTEXT("DeploySectionHeader", "Deploy").ToString())
+								.Text(LOCTEXT("DeploySectionHeader", "Deploy"))
 							]
 
 						+ SGridPanel::Slot(1, 8)
@@ -187,7 +187,7 @@ void SSessionLauncherSettings::Construct( const FArguments& InArgs, const FSessi
 							[
 								SNew(STextBlock)
 								.Font(FSlateFontInfo(FPaths::EngineContentDir() / TEXT("Slate/Fonts/Roboto-Bold.ttf"), 13))
-								.Text(LOCTEXT("LaunchSectionHeader", "Launch").ToString())
+								.Text(LOCTEXT("LaunchSectionHeader", "Launch"))
 							]
 
 						+ SGridPanel::Slot(1, 10)
@@ -203,7 +203,7 @@ void SSessionLauncherSettings::Construct( const FArguments& InArgs, const FSessi
 					.HAlign(HAlign_Center)
 					[
 						SNew(STextBlock)
-						.Text(LOCTEXT("SelectProfileText", "Select or create a new profile to continue.").ToString())
+						.Text(LOCTEXT("SelectProfileText", "Select or create a new profile to continue."))
 						.Visibility(this, &SSessionLauncherSettings::HandleSelectProfileTextBlockVisibility)
 					]
 			]
diff --git a/Engine/Source/Editor/DetailCustomizations/Private/AnimTransitionNodeDetails.cpp b/Engine/Source/Editor/DetailCustomizations/Private/AnimTransitionNodeDetails.cpp
--- a/Engine/Source/Editor/DetailCustomizations/Private/AnimTransitionNodeDetails.cpp
+++ b/Engine/Source/Editor/DetailCustomizations/Private/AnimTransitionNodeDetails.cpp
@@ -134,7 +134,7 @@ void FAnimTransitionNodeDetails::CustomizeDetails( IDetailLayoutBuilder& DetailB
 				.HAlign(HAlign_Right)
 				.OnClicked(this, &FAnimTransitionNodeDetails::OnClickEditBlendGraph)
 				.Visibility( this, &FAnimTransitionNodeDetails::GetBlendGraphButtonVisibility )
-				.Text(LOCTEXT("EditBlendGraph", "Edit Blend Graph").ToString())
+				.Text(LOCTEXT("EditBlendGraph", "Edit Blend Graph"))
 			]
 		];
 
diff --git a/Engine/Source/Editor/DetailCustomizations/Private/CollisionProfileDetails.cpp b/Engine/Source/Editor/DetailCustomizations/Private/CollisionProfileDetails.cpp
--- a/Engine/Source/Editor/DetailCustomizations/Private/CollisionProfileDetails.cpp
+++ b/Engine/Source/Editor/DetailCustomizations/Private/CollisionProfileDetails.cpp
@@ -49,7 +49,7 @@ public:
 	// widget event handlers
 	TSharedRef<SWidget> HandleResponseComboBoxGenerateWidget(TSharedPtr<FString> StringItem);
 	void HandleResponseComboBoxSelectionChanged(TSharedPtr<FString> StringItem, ESelectInfo::Type SelectInfo);
-	FString HandleResponseComboBoxContentText() const;
+	FText HandleResponseComboBoxContentText() const;
 
 	FText GetName() const;
 	void NewNameEntered(const FText& NewText, ETextCommit::Type CommitInfo);
@@ -118,7 +118,7 @@ void SChannelEditDialog::Construct(const FArguments& InArgs)
 				SNew(STextBlock)
 				.ColorAndOpacity(FSlateColor::UseSubduedForeground())
 				.Font(IDetailLayoutBuilder::GetDetailFont())
-				.Text(LOCTEXT("Dialog_TextTypeHelp", "Type Enter to apply text change. ").ToString())
+				.Text(LOCTEXT("Dialog_TextTypeHelp", "Type Enter to apply text change. "))
 			]
 			// channel name
 			+ SVerticalBox::Slot()
@@ -198,7 +198,7 @@ void SChannelEditDialog::Construct(const FArguments& InArgs)
 				[
 					SNew(SButton)
 					.ContentPadding(FEditorStyle::GetMargin("StandardDialog.ContentPadding"))
-					.Text(LOCTEXT("SChannelEditDialog_Accept", "Accept").ToString())
+					.Text(LOCTEXT("SChannelEditDialog_Accept", "Accept"))
 					.OnClicked(this, &SChannelEditDialog::OnAccept)
 					.IsEnabled(this, &SChannelEditDialog::IsAcceptAvailable)
 				]
@@ -206,7 +206,7 @@ void SChannelEditDialog::Construct(const FArguments& InArgs)
 				[
 					SNew(SButton)
 					.ContentPadding(FEditorStyle::GetMargin("StandardDialog.ContentPadding"))
-					.Text(LOCTEXT("SChannelEditDialog_Cancel", "Cancel").ToString())
+					.Text(LOCTEXT("SChannelEditDialog_Cancel", "Cancel"))
 					.OnClicked(this, &SChannelEditDialog::OnCancel)
 				]
 			]
@@ -264,11 +264,11 @@ void SChannelEditDialog::OnTextChanged(const FText& NewText)
 	if(NewName.Find(TEXT(" "))!=INDEX_NONE)
 	{
 		// no white space
-		NameBox->SetError(TEXT("No white space is allowed"));
+		NameBox->SetError(LOCTEXT("ChannelNameValidationWhitespaceError", "No white space is allowed"));
 	}
 	else
 	{
-		NameBox->SetError(TEXT(""));
+		NameBox->SetError(FText::GetEmpty());
 	}
 }
 
@@ -281,7 +281,7 @@ void SChannelEditDialog::NewNameEntered(const FText& NewText, ETextCommit::Type
 		{
 			ChannelSetup.Name = NewName;
 
-			NameBox->SetError(TEXT(""));
+			NameBox->SetError(FText::GetEmpty());
 		}
 	}
 }
@@ -290,7 +290,7 @@ FText SChannelEditDialog::GetName() const
 {
 	if (ChannelSetup.Name == NAME_None)
 	{
-		return FText::FromString(TEXT(""));
+		return FText::GetEmpty();
 	}
 
 	return FText::FromName(ChannelSetup.Name);
@@ -319,15 +319,15 @@ void SChannelEditDialog::HandleResponseComboBoxSelectionChanged(TSharedPtr<FStri
 	check(false);
 }
 
-FString SChannelEditDialog::HandleResponseComboBoxContentText() const
+FText SChannelEditDialog::HandleResponseComboBoxContentText() const
 {
 	int32 Index = (int32)ChannelSetup.DefaultResponse;
 	if (ResponseComboBoxString.IsValidIndex(Index))
 	{
-		return *ResponseComboBoxString[Index];
+		return FText::FromString(*ResponseComboBoxString[Index]);
 	}
 
-	return FString(TEXT("Select Response"));
+	return LOCTEXT("ChannelResponseTypeMessage", "Select Response");
 }
 
 //====================================================================================
@@ -357,11 +357,11 @@ public:
 	// widget event handlers
 	TSharedRef<SWidget> HandleCollisionEnabledComboBoxGenerateWidget(TSharedPtr<FString> StringItem);
 	void HandleCollisionEnabledComboBoxSelectionChanged(TSharedPtr<FString> StringItem, ESelectInfo::Type SelectInfo);
-	FString HandleCollisionEnabledComboBoxContentText() const;
+	FText HandleCollisionEnabledComboBoxContentText() const;
 
 	TSharedRef<SWidget> HandleObjectTypeComboBoxGenerateWidget(TSharedPtr<FString> StringItem);
 	void HandleObjectTypeComboBoxSelectionChanged(TSharedPtr<FString> StringItem, ESelectInfo::Type SelectInfo);
-	FString HandleObjectTypeComboBoxContentText() const;
+	FText HandleObjectTypeComboBoxContentText() const;
 
 	FText GetName() const;
 	void NewNameEntered(const FText& NewText, ETextCommit::Type CommitInfo);
@@ -769,15 +769,15 @@ void SProfileEditDialog::HandleCollisionEnabledComboBoxSelectionChanged(TSharedP
 	check(false);
 }
 
-FString SProfileEditDialog::HandleCollisionEnabledComboBoxContentText() const
+FText SProfileEditDialog::HandleCollisionEnabledComboBoxContentText() const
 {
 	int32 Index = (int32)ProfileTemplate.CollisionEnabled;
 	if(CollisionEnabledComboBoxString.IsValidIndex(Index))
 	{
-		return *CollisionEnabledComboBoxString[Index];
+		return FText::FromString(*CollisionEnabledComboBoxString[Index]);
 	}
 
-	return FString(TEXT("Select CollisionEnabled"));
+	return LOCTEXT("ProfileCollisionEnabledMessage", "Select CollisionEnabled");
 }
 
 
@@ -813,7 +813,7 @@ void SProfileEditDialog::HandleObjectTypeComboBoxSelectionChanged(TSharedPtr<FSt
 	check(false);
 }
 
-FString SProfileEditDialog::HandleObjectTypeComboBoxContentText() const
+FText SProfileEditDialog::HandleObjectTypeComboBoxContentText() const
 {
 	for(auto Iter = ObjectTypeMapping.CreateConstIterator(); Iter; ++Iter)
 	{
@@ -822,12 +822,12 @@ FString SProfileEditDialog::HandleObjectTypeComboBoxContentText() const
 			int32 Index = Iter.GetIndex();
 			if(ObjectTypeComboBoxString.IsValidIndex(Index))
 			{
-				return *ObjectTypeComboBoxString[Index];
+				return FText::FromString(*ObjectTypeComboBoxString[Index]);
 			}
 		}
 	}
 
-	return FString(TEXT("Select ObjectType"));
+	return LOCTEXT("ProfileObjectTypeMessage", "Select ObjectType");
 }
 
 void SProfileEditDialog::AddCollisionResponse()
diff --git a/Engine/Source/Editor/DetailCustomizations/Private/CurveStructCustomization.cpp b/Engine/Source/Editor/DetailCustomizations/Private/CurveStructCustomization.cpp
--- a/Engine/Source/Editor/DetailCustomizations/Private/CurveStructCustomization.cpp
+++ b/Engine/Source/Editor/DetailCustomizations/Private/CurveStructCustomization.cpp
@@ -98,7 +98,7 @@ void FCurveStructCustomization::CustomizeStructHeader( TSharedRef<class IPropert
 				.VAlign(VAlign_Fill)
 				[
 					SNew(STextBlock)
-					.Text(LOCTEXT("MultipleCurves", "Multiple Curves - unable to modify").ToString())
+					.Text(LOCTEXT("MultipleCurves", "Multiple Curves - unable to modify"))
 				]
 			];
 	}
@@ -144,8 +144,8 @@ void FCurveStructCustomization::CustomizeStructChildren( TSharedRef<class IPrope
 							SNew(SButton)
 							.ButtonStyle( FEditorStyle::Get(), "NoBorder" )
 							.ContentPadding(1.f)
-							.ToolTipText(LOCTEXT( "ConvertInternalCurveTooltip", "Convert to Internal Curve").ToString())
-							//.ToolTipText(LOCTEXT( "ConvertInternalCurveTooltip", "Copy the external CurveFloat asset to this curve").ToString() )
+							.ToolTipText(LOCTEXT( "ConvertInternalCurveTooltip", "Convert to Internal Curve"))
+							//.ToolTipText(LOCTEXT( "ConvertInternalCurveTooltip", "Copy the external CurveFloat asset to this curve") )
 							.OnClicked(this, &FCurveStructCustomization::OnConvertButtonClicked)
 							.IsEnabled(this, &FCurveStructCustomization::IsConvertButtonEnabled)
 							[
@@ -162,8 +162,8 @@ void FCurveStructCustomization::CustomizeStructChildren( TSharedRef<class IPrope
 						[
 							SNew(SButton)
 							.HAlign(HAlign_Center)
-							.Text( LOCTEXT( "CreateAssetButton", "Create External Curve" ).ToString() )
-							.ToolTipText(LOCTEXT( "CreateAssetTooltip", "Create a new CurveFloat asset from this curve").ToString() )
+							.Text( LOCTEXT( "CreateAssetButton", "Create External Curve" ) )
+							.ToolTipText(LOCTEXT( "CreateAssetTooltip", "Create a new CurveFloat asset from this curve") )
 							.OnClicked(this, &FCurveStructCustomization::OnCreateButtonClicked)
 							.IsEnabled(this, &FCurveStructCustomization::IsCreateButtonEnabled)
 						]
@@ -171,8 +171,8 @@ void FCurveStructCustomization::CustomizeStructChildren( TSharedRef<class IPrope
 						[
 							SNew(SButton)
 							.HAlign(HAlign_Center)
-							.Text( LOCTEXT( "ConvertAssetButton", "Convert to Internal Curve" ).ToString() )
-							.ToolTipText(LOCTEXT( "CopyExternalCurveTooltip", "Copy the external CurveFloat asset to this curve").ToString() )
+							.Text( LOCTEXT( "ConvertAssetButton", "Convert to Internal Curve" ) )
+							.ToolTipText(LOCTEXT( "CopyExternalCurveTooltip", "Copy the external CurveFloat asset to this curve") )
 							.OnClicked(this, &FCurveStructCustomization::OnConvertButtonClicked)
 							.IsEnabled(this, &FCurveStructCustomization::IsConvertButtonEnabled)
 						]*/
diff --git a/Engine/Source/Editor/DetailCustomizations/Private/DataTableCategoryCustomization.h b/Engine/Source/Editor/DetailCustomizations/Private/DataTableCategoryCustomization.h
--- a/Engine/Source/Editor/DetailCustomizations/Private/DataTableCategoryCustomization.h
+++ b/Engine/Source/Editor/DetailCustomizations/Private/DataTableCategoryCustomization.h
@@ -51,7 +51,7 @@ public:
 			.NameContent()
 				[
 					SNew(STextBlock)
-					.Text(LOCTEXT("DataTable_ColumnName", "Column Name").ToString())
+					.Text(LOCTEXT("DataTable_ColumnName", "Column Name"))
 					.Font(StructCustomizationUtils.GetRegularFont())
 				]
 			.ValueContent()
@@ -75,7 +75,7 @@ public:
 			.NameContent()
 				[
 					SNew(STextBlock)
-					.Text(LOCTEXT("DataTable_RowContains", "Row Contains").ToString())
+					.Text(LOCTEXT("DataTable_RowContains", "Row Contains"))
 					.Font(StructCustomizationUtils.GetRegularFont())
 				]
 			.ValueContent()
diff --git a/Engine/Source/Editor/DetailCustomizations/Private/DialogueWaveDetails.cpp b/Engine/Source/Editor/DetailCustomizations/Private/DialogueWaveDetails.cpp
--- a/Engine/Source/Editor/DetailCustomizations/Private/DialogueWaveDetails.cpp
+++ b/Engine/Source/Editor/DetailCustomizations/Private/DialogueWaveDetails.cpp
@@ -104,8 +104,8 @@ void FDialogueWaveDetails::CustomizeDetails( class IDetailLayoutBuilder& DetailB
 		.HAlign(HAlign_Center)
 		[
 			SNew( SButton )
-			.Text( LOCTEXT("AddDialogueContext", "Add Dialogue Context").ToString() )
-			.ToolTipText( LOCTEXT("AddDialogueContextToolTip", "Adds a new context for dialogue based on speakers, those spoken to, and the associated soundwave.").ToString() )
+			.Text( LOCTEXT("AddDialogueContext", "Add Dialogue Context") )
+			.ToolTipText( LOCTEXT("AddDialogueContextToolTip", "Adds a new context for dialogue based on speakers, those spoken to, and the associated soundwave.") )
 			.OnClicked( FOnClicked::CreateSP( this, &FDialogueWaveDetails::AddDialogueContextMapping_OnClicked ) )
 		]
 	];
diff --git a/Engine/Source/Editor/DetailCustomizations/Private/LandscapeSplineDetails.cpp b/Engine/Source/Editor/DetailCustomizations/Private/LandscapeSplineDetails.cpp
--- a/Engine/Source/Editor/DetailCustomizations/Private/LandscapeSplineDetails.cpp
+++ b/Engine/Source/Editor/DetailCustomizations/Private/LandscapeSplineDetails.cpp
@@ -25,13 +25,13 @@ void FLandscapeSplineDetails::CustomizeDetails( IDetailLayoutBuilder& DetailBuil
 		.FillWidth(1)
 		[
 			SNew(STextBlock)
-			.Text(LOCTEXT("SelectAll", "Select all connected:").ToString())
+			.Text(LOCTEXT("SelectAll", "Select all connected:"))
 		]
 		+ SHorizontalBox::Slot()
 		.FillWidth(1)
 		[
 			SNew(SButton)
-			.Text(LOCTEXT("ControlPoints", "Control Points").ToString())
+			.Text(LOCTEXT("ControlPoints", "Control Points"))
 			.HAlign(HAlign_Center)
 			.OnClicked(this, &FLandscapeSplineDetails::OnSelectConnectedControlPointsButtonClicked)
 		]
@@ -39,7 +39,7 @@ void FLandscapeSplineDetails::CustomizeDetails( IDetailLayoutBuilder& DetailBuil
 		.FillWidth(1)
 		[
 			SNew(SButton)
-			.Text(LOCTEXT("Segments", "Segments").ToString())
+			.Text(LOCTEXT("Segments", "Segments"))
 			.HAlign(HAlign_Center)
 			.OnClicked(this, &FLandscapeSplineDetails::OnSelectConnectedSegmentsButtonClicked)
 		]
diff --git a/Engine/Source/Editor/DetailCustomizations/Private/PhysicsConstraintComponentDetails.cpp b/Engine/Source/Editor/DetailCustomizations/Private/PhysicsConstraintComponentDetails.cpp
--- a/Engine/Source/Editor/DetailCustomizations/Private/PhysicsConstraintComponentDetails.cpp
+++ b/Engine/Source/Editor/DetailCustomizations/Private/PhysicsConstraintComponentDetails.cpp
@@ -63,8 +63,8 @@ void FPhysicsConstraintComponentDetails::CustomizeDetails( IDetailLayoutBuilder&
 						SNew(SButton)
 						.VAlign(VAlign_Center)
 						.OnClicked( this, &FPhysicsConstraintComponentDetails::OnHingeClicked )
-						.Text( LOCTEXT("HingePreset", "Hinge").ToString() )
-						.ToolTipText( LOCTEXT("HingePresetTooltip", "Setup joint as hinge. A hinge joint allows motion only in one plane.").ToString() )
+						.Text( LOCTEXT("HingePreset", "Hinge") )
+						.ToolTipText( LOCTEXT("HingePresetTooltip", "Setup joint as hinge. A hinge joint allows motion only in one plane.") )
 					]
 					+SHorizontalBox::Slot()
 						.AutoWidth()
@@ -75,8 +75,8 @@ void FPhysicsConstraintComponentDetails::CustomizeDetails( IDetailLayoutBuilder&
 							SNew(SButton)
 							.VAlign(VAlign_Center)
 							.OnClicked( this, &FPhysicsConstraintComponentDetails::OnPrismaticClicked )
-							.Text( LOCTEXT("PrismaticPreset", "Prismatic").ToString() )
-							.ToolTipText( LOCTEXT("PrismaticPresetTooltip", "Setup joint as prismatic. A prismatic joint allows only linear sliding movement.").ToString() )
+							.Text( LOCTEXT("PrismaticPreset", "Prismatic") )
+							.ToolTipText( LOCTEXT("PrismaticPresetTooltip", "Setup joint as prismatic. A prismatic joint allows only linear sliding movement.") )
 						]
 					+SHorizontalBox::Slot()
 						.AutoWidth()
@@ -87,8 +87,8 @@ void FPhysicsConstraintComponentDetails::CustomizeDetails( IDetailLayoutBuilder&
 							SNew(SButton)
 							.VAlign(VAlign_Center)
 							.OnClicked( this, &FPhysicsConstraintComponentDetails::OnBallSocketClicked )
-							.Text( LOCTEXT("BSPreset", "Ball and Socket").ToString() )
-							.ToolTipText( LOCTEXT("BSPresetTooltip", "Setup joint as ball and socket. A Ball and Socket joint allows motion around an indefinite number of axes, which have one common center").ToString() )
+							.Text( LOCTEXT("BSPreset", "Ball and Socket") )
+							.ToolTipText( LOCTEXT("BSPresetTooltip", "Setup joint as ball and socket. A Ball and Socket joint allows motion around an indefinite number of axes, which have one common center") )
 						]
 				]
 			];
diff --git a/Engine/Source/Editor/DetailCustomizations/Private/SkeletalControlNodeDetails.cpp b/Engine/Source/Editor/DetailCustomizations/Private/SkeletalControlNodeDetails.cpp
--- a/Engine/Source/Editor/DetailCustomizations/Private/SkeletalControlNodeDetails.cpp
+++ b/Engine/Source/Editor/DetailCustomizations/Private/SkeletalControlNodeDetails.cpp
@@ -48,7 +48,7 @@ void FSkeletalControlNodeDetails::OnGenerateElementForPropertyPin(TSharedRef<IPr
 		case FPropertyAccess::MultipleValues:
 			ChildrenBuilder.AddChildContent( TEXT("") ) 
 			[
-				 SNew(STextBlock).Text(LOCTEXT("OnlyWorksInSingleSelectMode", "Multiple types selected").ToString())
+				 SNew(STextBlock).Text(LOCTEXT("OnlyWorksInSingleSelectMode", "Multiple types selected"))
 			];
 			return;
 		case FPropertyAccess::Fail:
diff --git a/Engine/Source/Editor/DetailCustomizations/Private/SkeletonNotifyDetails.cpp b/Engine/Source/Editor/DetailCustomizations/Private/SkeletonNotifyDetails.cpp
--- a/Engine/Source/Editor/DetailCustomizations/Private/SkeletonNotifyDetails.cpp
+++ b/Engine/Source/Editor/DetailCustomizations/Private/SkeletonNotifyDetails.cpp
@@ -38,8 +38,8 @@ void FSkeletonNotifyDetails::CustomizeDetails( IDetailLayoutBuilder& DetailBuild
 		.NameContent()
 		[
 			SNew(STextBlock)
-			.ToolTipText(LOCTEXT("Animations_Tooltip", "List of animations that reference this notify").ToString())
-			.Text( LOCTEXT("AnimationsLabel","Animations").ToString() )
+			.ToolTipText(LOCTEXT("Animations_Tooltip", "List of animations that reference this notify"))
+			.Text( LOCTEXT("AnimationsLabel","Animations") )
 			.Font( DetailFontInfo )
 		]
 		.ValueContent()
diff --git a/Engine/Source/Editor/DetailCustomizations/Private/SkinnedMeshComponentDetails.cpp b/Engine/Source/Editor/DetailCustomizations/Private/SkinnedMeshComponentDetails.cpp
--- a/Engine/Source/Editor/DetailCustomizations/Private/SkinnedMeshComponentDetails.cpp
+++ b/Engine/Source/Editor/DetailCustomizations/Private/SkinnedMeshComponentDetails.cpp
@@ -33,7 +33,7 @@ void FSkinnedMeshComponentDetails::CreateActuallyUsedPhysicsAssetWidget(FDetailW
 	.NameContent()
 	[
 		SNew(STextBlock)
-		.Text(LOCTEXT("CurrentPhysicsAsset", "Currently used Physics Asset").ToString())
+		.Text(LOCTEXT("CurrentPhysicsAsset", "Currently used Physics Asset"))
 		.Font(IDetailLayoutBuilder::GetDetailFont())
 	]
 	.ValueContent()
@@ -92,7 +92,7 @@ FText FSkinnedMeshComponentDetails::GetUsedPhysicsAssetAsText(IDetailLayoutBuild
 	UPhysicsAsset* UsedPhysicsAsset = NULL;
 	if (! FindUniqueUsedPhysicsAsset(DetailBuilder, UsedPhysicsAsset))
 	{
-		return FText::FromString(LOCTEXT("MultipleValues", "Multiple Values").ToString());
+		return LOCTEXT("MultipleValues", "Multiple Values");
 	}
 	if (UsedPhysicsAsset)
 	{
@@ -100,7 +100,7 @@ FText FSkinnedMeshComponentDetails::GetUsedPhysicsAssetAsText(IDetailLayoutBuild
 	}
 	else
 	{
-		return FText::FromString(TEXT(""));
+		return FText::GetEmpty();
 	}
 }
 
diff --git a/Engine/Source/Editor/DeviceProfileEditor/Private/DetailsPanel/SDeviceProfileDetailsPanel.cpp b/Engine/Source/Editor/DeviceProfileEditor/Private/DetailsPanel/SDeviceProfileDetailsPanel.cpp
--- a/Engine/Source/Editor/DeviceProfileEditor/Private/DetailsPanel/SDeviceProfileDetailsPanel.cpp
+++ b/Engine/Source/Editor/DeviceProfileEditor/Private/DetailsPanel/SDeviceProfileDetailsPanel.cpp
@@ -42,7 +42,7 @@ void SDeviceProfileDetailsPanel::Construct( const FArguments& InArgs )
 					.VAlign( VAlign_Center )
 					[
 						SNew( STextBlock )
-						.Text( LOCTEXT("CVarsLabel", "Console Variables").ToString() )
+						.Text( LOCTEXT("CVarsLabel", "Console Variables") )
 						.TextStyle( FEditorStyle::Get(), "Docking.TabFont" )
 					]
 			]
@@ -170,7 +170,7 @@ void SDeviceProfileDetailsPanel::RefreshUI()
 					.Padding( 4.0f )
 					[
 						SNew( STextBlock )
-						.Text( LOCTEXT("SelectAProfile", "Select a device profile above...").ToString() )					
+						.Text( LOCTEXT("SelectAProfile", "Select a device profile above...") )					
 					]
 
 				]
diff --git a/Engine/Source/Editor/DistCurveEditor/Private/SDistributionCurveEditor.cpp b/Engine/Source/Editor/DistCurveEditor/Private/SDistributionCurveEditor.cpp
--- a/Engine/Source/Editor/DistCurveEditor/Private/SDistributionCurveEditor.cpp
+++ b/Engine/Source/Editor/DistCurveEditor/Private/SDistributionCurveEditor.cpp
@@ -311,7 +311,7 @@ TSharedRef<SHorizontalBox> SDistributionCurveEditor::BuildToolBar()
 				.Padding(4)
 				[
 					SNew(STextBlock)
-					.Text(LOCTEXT("CurrentTab", "Current Tab: ").ToString())
+					.Text(LOCTEXT("CurrentTab", "Current Tab: "))
 					.Visibility( this, &SDistributionCurveEditor::GetLargeIconVisibility )
 				]
 				+SVerticalBox::Slot()
@@ -519,7 +519,7 @@ void SDistributionCurveEditor::OnSetTime()
 
 	TSharedRef<STextEntryPopup> TextEntry = 
 		SNew(STextEntryPopup)
-		.Label(LOCTEXT("SetTime", "Time: ").ToString())
+		.Label(LOCTEXT("SetTime", "Time: "))
 		.DefaultText(FText::FromString( DefaultText ) )
 		.OnTextCommitted(this, &SDistributionCurveEditor::KeyTimeCommitted)
 		.SelectAllTextWhenFocused(true)
@@ -545,7 +545,7 @@ void SDistributionCurveEditor::OnSetValue()
 
 	TSharedRef<STextEntryPopup> TextEntry = 
 		SNew(STextEntryPopup)
-		.Label(LOCTEXT("SetValue", "Value: ").ToString())
+		.Label(LOCTEXT("SetValue", "Value: "))
 		.DefaultText(FText::FromString(DefaultText))
 		.OnTextCommitted(this, &SDistributionCurveEditor::KeyValueCommitted)
 		.SelectAllTextWhenFocused(true)
@@ -767,7 +767,7 @@ void SDistributionCurveEditor::OnScaleTimes(ECurveScaleScope::Type Scope)
 
 	TSharedRef<STextEntryPopup> TextEntry = 
 		SNew(STextEntryPopup)
-		.Label(Label.ToString())
+		.Label(Label)
 		.DefaultText(FText::FromString( DefaultText ))
 		.OnTextCommitted(this, &SDistributionCurveEditor::ScaleTimeCommitted, Scope)
 		.SelectAllTextWhenFocused(true)
@@ -805,7 +805,7 @@ void SDistributionCurveEditor::OnScaleValues(ECurveScaleScope::Type Scope)
 
 	TSharedRef<STextEntryPopup> TextEntry = 
 		SNew(STextEntryPopup)
-		.Label(Label.ToString())
+		.Label(Label)
 		.DefaultText(FText::FromString( DefaultText ))
 		.OnTextCommitted(this, &SDistributionCurveEditor::ScaleValueCommitted, Scope)
 		.SelectAllTextWhenFocused(true)
@@ -964,7 +964,7 @@ void SDistributionCurveEditor::OnCreateTab()
 {
 	TSharedRef<STextEntryPopup> TextEntry = 
 		SNew(STextEntryPopup)
-		.Label(LOCTEXT("SetTabName", "Tab Name: ").ToString())
+		.Label(LOCTEXT("SetTabName", "Tab Name: "))
 		.OnTextCommitted(this, &SDistributionCurveEditor::TabNameCommitted)
 		.ClearKeyboardFocusOnCommit(false);
 
diff --git a/Engine/Source/Editor/Documentation/Private/UDNParser.cpp b/Engine/Source/Editor/Documentation/Private/UDNParser.cpp
--- a/Engine/Source/Editor/Documentation/Private/UDNParser.cpp
+++ b/Engine/Source/Editor/Documentation/Private/UDNParser.cpp
@@ -954,7 +954,7 @@ TSharedRef< SWidget > FUDNParser::GenerateExcerptContent( const FString& Link, F
 
 	if ( bCriticalError )
 	{
-		return SNew( STextBlock ).Text( LOCTEXT("ExcerptContentLoadingError", "Excerpt {0} could not be loaded.  :(").ToString() );
+		return SNew( STextBlock ).Text( LOCTEXT("ExcerptContentLoadingError", "Excerpt {0} could not be loaded.  :(") );
 	}
 
 	return Box.ToSharedRef();
diff --git a/Engine/Source/Editor/FoliageEdit/Private/SFoliageEdit.cpp b/Engine/Source/Editor/FoliageEdit/Private/SFoliageEdit.cpp
--- a/Engine/Source/Editor/FoliageEdit/Private/SFoliageEdit.cpp
+++ b/Engine/Source/Editor/FoliageEdit/Private/SFoliageEdit.cpp
@@ -106,7 +106,7 @@ void SFoliageEdit::Construct(const FArguments& InArgs)
 					[
 						SNew(STextBlock)
 							.Visibility( this, &SFoliageEdit::GetVisibility_Radius )
-							.Text(LOCTEXT("BrushSize", "Brush Size").ToString())
+							.Text(LOCTEXT("BrushSize", "Brush Size"))
 					]
 					+SHorizontalBox::Slot()
 						.FillWidth(2.0f)
@@ -134,7 +134,7 @@ void SFoliageEdit::Construct(const FArguments& InArgs)
 					[
 						SNew(STextBlock)
 							.Visibility( this, &SFoliageEdit::GetVisibility_PaintDensity )
-							.Text(LOCTEXT("PaintDensity", "Paint Density").ToString())
+							.Text(LOCTEXT("PaintDensity", "Paint Density"))
 					]
 					+SHorizontalBox::Slot()
 						.FillWidth(2.0f)
@@ -161,7 +161,7 @@ void SFoliageEdit::Construct(const FArguments& InArgs)
 					[
 						SNew(STextBlock)
 							.Visibility( this, &SFoliageEdit::GetVisibility_EraseDensity )
-							.Text(LOCTEXT("EraseDensity", "Erase Density").ToString())
+							.Text(LOCTEXT("EraseDensity", "Erase Density"))
 					]
 					+SHorizontalBox::Slot()
 						.FillWidth(2.0f)
@@ -183,7 +183,7 @@ void SFoliageEdit::Construct(const FArguments& InArgs)
 				[
 					SNew(STextBlock)
 						.Visibility( this, &SFoliageEdit::GetVisibility_Filters )
-						.Text(LOCTEXT("Filter", "Filter").ToString())
+						.Text(LOCTEXT("Filter", "Filter"))
 				]
 
 				+SVerticalBox::Slot()
@@ -202,7 +202,7 @@ void SFoliageEdit::Construct(const FArguments& InArgs)
 							.IsChecked( this, &SFoliageEdit::GetCheckState_Landscape )
 							[
 								SNew(STextBlock)
-									.Text(LOCTEXT("Landscape", "Landscape").ToString())
+									.Text(LOCTEXT("Landscape", "Landscape"))
 							]
 					]
 				
@@ -215,7 +215,7 @@ void SFoliageEdit::Construct(const FArguments& InArgs)
 							.IsChecked( this, &SFoliageEdit::GetCheckState_StaticMesh )
 							[
 								SNew(STextBlock)
-									.Text(LOCTEXT("StaticMeshes", "Static Meshes").ToString())
+									.Text(LOCTEXT("StaticMeshes", "Static Meshes"))
 							]
 					]
 
@@ -227,7 +227,7 @@ void SFoliageEdit::Construct(const FArguments& InArgs)
 							.IsChecked( this, &SFoliageEdit::GetCheckState_BSP )
 							[
 								SNew(STextBlock)
-									.Text(LOCTEXT("BSP", "BSP").ToString())
+									.Text(LOCTEXT("BSP", "BSP"))
 							]
 					]
 				]
@@ -238,7 +238,7 @@ void SFoliageEdit::Construct(const FArguments& InArgs)
 					SNew(SHeaderRow)
 
 					+SHeaderRow::Column(TEXT("Meshes"))
-						.DefaultLabel(LOCTEXT("Meshes", "Meshes").ToString())
+						.DefaultLabel(LOCTEXT("Meshes", "Meshes"))
 				]
 
 				+SVerticalBox::Slot()
@@ -284,7 +284,7 @@ void SFoliageEdit::Construct(const FArguments& InArgs)
 								[
 									SNew(STextBlock)
 										.Visibility( this, &SFoliageEdit::GetVisibility_EmptyListText )
-										.Text( LOCTEXT("EmptyFoliageListMessage", "Drag static meshes from the content browser into this area.").ToString() )	
+										.Text( LOCTEXT("EmptyFoliageListMessage", "Drag static meshes from the content browser into this area.") )	
 								]
 
 								+SHorizontalBox::Slot()
diff --git a/Engine/Source/Editor/GameProjectGeneration/Private/SNewClassDialog.cpp b/Engine/Source/Editor/GameProjectGeneration/Private/SNewClassDialog.cpp
--- a/Engine/Source/Editor/GameProjectGeneration/Private/SNewClassDialog.cpp
+++ b/Engine/Source/Editor/GameProjectGeneration/Private/SNewClassDialog.cpp
@@ -119,7 +119,7 @@ void SNewClassDialog::Construct( const FArguments& InArgs )
 						[
 							SNew(STextBlock)
 							.TextStyle( FEditorStyle::Get(), "NewClassDialog.PageTitle" )
-							.Text( LOCTEXT( "ParentClassTitle", "Choose Parent Class" ).ToString() )
+							.Text( LOCTEXT( "ParentClassTitle", "Choose Parent Class" ) )
 						]
 
 						// Title spacer
@@ -142,7 +142,7 @@ void SNewClassDialog::Construct( const FArguments& InArgs )
 							.VAlign(VAlign_Center)
 							[
 								SNew(STextBlock)
-								.Text( FText::Format( LOCTEXT("ChooseParentClassDescription", "You are about to add a C++ source code file. To compile these files you must have {0} installed."), FSourceCodeNavigation::GetSuggestedSourceCodeIDE() ).ToString() )
+								.Text( FText::Format( LOCTEXT("ChooseParentClassDescription", "You are about to add a C++ source code file. To compile these files you must have {0} installed."), FSourceCodeNavigation::GetSuggestedSourceCodeIDE() ) )
 							]
 
 							// Full tree checkbox
@@ -156,7 +156,7 @@ void SNewClassDialog::Construct( const FArguments& InArgs )
 								.OnCheckStateChanged( this, &SNewClassDialog::OnFullClassTreeChanged )
 								[
 									SNew(STextBlock)
-									.Text( LOCTEXT( "FullClassTree", "Show All Classes" ).ToString() )
+									.Text( LOCTEXT( "FullClassTree", "Show All Classes" ) )
 								]
 							]
 						]
@@ -211,7 +211,7 @@ void SNewClassDialog::Construct( const FArguments& InArgs )
 							[
 								SNew(STextBlock)
 								.TextStyle( FEditorStyle::Get(), "NewClassDialog.SelectedParentClassLabel" )
-								.Text( LOCTEXT( "ParentClassLabel", "Selected Class" ).ToString() )
+								.Text( LOCTEXT( "ParentClassLabel", "Selected Class" ) )
 							]
 
 							// Class selection preview
@@ -261,7 +261,7 @@ void SNewClassDialog::Construct( const FArguments& InArgs )
 							.Padding(0, 0, 0, 5)
 							[
 								SNew(STextBlock)
-								.Text( LOCTEXT("ClassNameDescription", "Enter a name for your new class. Class names may only contain alphanumeric characters, and may not contain a space.").ToString() )
+								.Text( LOCTEXT("ClassNameDescription", "Enter a name for your new class. Class names may only contain alphanumeric characters, and may not contain a space.") )
 							]
 
 							+SVerticalBox::Slot()
@@ -269,7 +269,7 @@ void SNewClassDialog::Construct( const FArguments& InArgs )
 							.Padding(0, 0, 0, 20)
 							[
 								SNew(STextBlock)
-								.Text( LOCTEXT("ClassNameDetails", "When you click the \"Create\" button below, a header (.h) file and a source (.cpp) file will be made using this name.").ToString() )
+								.Text( LOCTEXT("ClassNameDetails", "When you click the \"Create\" button below, a header (.h) file and a source (.cpp) file will be made using this name.") )
 							]
 
 							// Name Error label
@@ -306,7 +306,7 @@ void SNewClassDialog::Construct( const FArguments& InArgs )
 								[
 									SNew(STextBlock)
 									.TextStyle( FEditorStyle::Get(), "NewClassDialog.SelectedParentClassLabel" )
-									.Text( LOCTEXT( "NameLabel", "Name" ).ToString() )
+									.Text( LOCTEXT( "NameLabel", "Name" ) )
 								]
 
 								// Name edit box
@@ -347,7 +347,7 @@ void SNewClassDialog::Construct( const FArguments& InArgs )
 									.AutoWidth()
 									[
 										SNew(SHyperlink)
-											.Text( FText::Format( LOCTEXT("IDEDownloadLinkText", "Download {0}"), FSourceCodeNavigation::GetSuggestedSourceCodeIDE() ).ToString() )
+											.Text( FText::Format( LOCTEXT("IDEDownloadLinkText", "Download {0}"), FSourceCodeNavigation::GetSuggestedSourceCodeIDE() ) )
 											.OnNavigate( this, &SNewClassDialog::OnDownloadIDEClicked, FSourceCodeNavigation::GetSuggestedSourceCodeIDEDownloadURL() )
 											.Visibility( this, &SNewClassDialog::GetGlobalErrorLabelIDELinkVisibility )
 									]
diff --git a/Engine/Source/Editor/GameplayTagsEditor/Private/SGameplayTagWidget.cpp b/Engine/Source/Editor/GameplayTagsEditor/Private/SGameplayTagWidget.cpp
--- a/Engine/Source/Editor/GameplayTagsEditor/Private/SGameplayTagWidget.cpp
+++ b/Engine/Source/Editor/GameplayTagsEditor/Private/SGameplayTagWidget.cpp
@@ -46,14 +46,14 @@ void SGameplayTagWidget::Construct(const FArguments& InArgs, const TArray<FEdita
 				[
 					SNew(SButton)
 					.OnClicked(this, &SGameplayTagWidget::OnExpandAllClicked)
-					.Text(LOCTEXT("GameplayTagWidget_ExpandAll", "Expand All").ToString())
+					.Text(LOCTEXT("GameplayTagWidget_ExpandAll", "Expand All"))
 				]
 				+SHorizontalBox::Slot()
 				.AutoWidth()
 				[
 					SNew(SButton)
 					.OnClicked(this, &SGameplayTagWidget::OnCollapseAllClicked)
-					.Text(LOCTEXT("GameplayTagWidget_CollapseAll", "Collapse All").ToString())
+					.Text(LOCTEXT("GameplayTagWidget_CollapseAll", "Collapse All"))
 				]
 				+SHorizontalBox::Slot()
 				.AutoWidth()
@@ -61,7 +61,7 @@ void SGameplayTagWidget::Construct(const FArguments& InArgs, const TArray<FEdita
 					SNew(SButton)
 					.IsEnabled(!bReadOnly)
 					.OnClicked(this, &SGameplayTagWidget::OnClearAllClicked)
-					.Text(LOCTEXT("GameplayTagWidget_ClearAll", "Clear All").ToString())
+					.Text(LOCTEXT("GameplayTagWidget_ClearAll", "Clear All"))
 				]
 				+SHorizontalBox::Slot()
 				.VAlign( VAlign_Center )
@@ -151,10 +151,10 @@ bool SGameplayTagWidget::FilterChildrenCheck( TSharedPtr<FGameplayTagNode> InIte
 
 TSharedRef<ITableRow> SGameplayTagWidget::OnGenerateRow(TSharedPtr<FGameplayTagNode> InItem, const TSharedRef<STableViewBase>& OwnerTable)
 {
-	FString TooltipText;
+	FText TooltipText;
 	if (InItem.IsValid())
 	{
-		TooltipText = InItem.Get()->GetCompleteTag().ToString();
+		TooltipText = FText::FromName(InItem.Get()->GetCompleteTag());
 	}
 
 	return SNew(STableRow< TSharedPtr<FGameplayTagNode> >, OwnerTable)
@@ -167,7 +167,7 @@ TSharedRef<ITableRow> SGameplayTagWidget::OnGenerateRow(TSharedPtr<FGameplayTagN
 			.ReadOnly( bReadOnly )
 			[
 				SNew(STextBlock)				
-				.Text(InItem->GetSimpleTag().ToString())
+				.Text(FText::FromName(InItem->GetSimpleTag()))
 			]	
 		];
 }
diff --git a/Engine/Source/Editor/GeometryMode/Private/GeometryMode.cpp b/Engine/Source/Editor/GeometryMode/Private/GeometryMode.cpp
--- a/Engine/Source/Editor/GeometryMode/Private/GeometryMode.cpp
+++ b/Engine/Source/Editor/GeometryMode/Private/GeometryMode.cpp
@@ -172,7 +172,7 @@ TSharedRef<SVerticalBox> SGeometryModeControls::CreateTopModifierButtons()
 	.HAlign(HAlign_Center)
 	[
 		SNew(SButton)
-		.Text(LOCTEXT("SGeometryModeDialog_Apply", "Apply").ToString())
+		.Text(LOCTEXT("SGeometryModeDialog_Apply", "Apply"))
 		.OnClicked(this, &SGeometryModeControls::OnApplyClicked)
 	];
 
@@ -239,7 +239,7 @@ TSharedRef<SCheckBox> SGeometryModeControls::CreateSingleModifierRadioButton(UGe
 	.IsEnabled(this, &SGeometryModeControls::IsModifierEnabled, Modifier)
 	.OnCheckStateChanged(this, &SGeometryModeControls::OnModifierStateChanged, Modifier)
 	[
-		SNew(STextBlock).Text( Modifier->GetModifierDescription().ToString() )
+		SNew(STextBlock).Text( Modifier->GetModifierDescription() )
 	];
 
 	ModifierControls.Add(CheckBox);
@@ -251,7 +251,7 @@ TSharedRef<SButton> SGeometryModeControls::CreateSingleModifierButton(UGeomModif
 {
 	TSharedRef<SButton> Widget =
 	SNew(SButton)
-	.Text( Modifier->GetModifierDescription().ToString() )
+	.Text( Modifier->GetModifierDescription() )
 	.HAlign(HAlign_Center)
 	.IsEnabled(this, &SGeometryModeControls::IsModifierEnabled, Modifier)
 	.OnClicked(this, &SGeometryModeControls::OnModifierClicked, Modifier);
diff --git a/Engine/Source/Editor/InputBindingEditor/Private/Widgets/SGestureEditBox.cpp b/Engine/Source/Editor/InputBindingEditor/Private/Widgets/SGestureEditBox.cpp
--- a/Engine/Source/Editor/InputBindingEditor/Private/Widgets/SGestureEditBox.cpp
+++ b/Engine/Source/Editor/InputBindingEditor/Private/Widgets/SGestureEditBox.cpp
@@ -59,7 +59,7 @@ void SGestureEditBox::Construct( const FArguments& InArgs, TSharedPtr<FGestureTr
 						.OnClicked( this, &SGestureEditBox::OnGestureRemoveButtonClicked )
 						.ForegroundColor( FSlateColor::UseForeground() )
 						.IsFocusable(false)
-						.ToolTipText(LOCTEXT("GestureEditButtonRemove_ToolTip", "Remove this binding").ToString() )
+						.ToolTipText(LOCTEXT("GestureEditButtonRemove_ToolTip", "Remove this binding") )
 						[
 							SNew( SBox )
 							[
@@ -215,7 +215,7 @@ TSharedRef<SWidget> SGestureEditBox::OnGetContentForConflictPopup()
 				[
 					SAssignNew( GestureAcceptButton, SButton )
 						.ContentPadding(1)
-						.ToolTipText( LOCTEXT("GestureAcceptButton_ToolTip", "Accept this new binding").ToString() )
+						.ToolTipText( LOCTEXT("GestureAcceptButton_ToolTip", "Accept this new binding") )
 						.OnClicked( this, &SGestureEditBox::OnAcceptNewGestureButtonClicked )
 						[
 							SNew( SHorizontalBox )
@@ -233,7 +233,7 @@ TSharedRef<SWidget> SGestureEditBox::OnGetContentForConflictPopup()
 								.Padding(2.0f,0.0f)
 								[
 									SNew( STextBlock )
-										.Text( LOCTEXT("GestureAcceptButtonText_Override", "Override").ToString() )
+										.Text( LOCTEXT("GestureAcceptButtonText_Override", "Override") )
 								]
 						]
 				]
diff --git a/Engine/Source/Editor/Kismet/Private/BlueprintDetailsCustomization.cpp b/Engine/Source/Editor/Kismet/Private/BlueprintDetailsCustomization.cpp
--- a/Engine/Source/Editor/Kismet/Private/BlueprintDetailsCustomization.cpp
+++ b/Engine/Source/Editor/Kismet/Private/BlueprintDetailsCustomization.cpp
@@ -61,7 +61,7 @@ void FBlueprintVarActionDetails::CustomizeDetails( IDetailLayoutBuilder& DetailL
 	.NameContent()
 	[
 		SNew(STextBlock)
-		.Text(LOCTEXT("BlueprintVarActionDetails_VariableNameLabel", "Variable Name").ToString())
+		.Text(LOCTEXT("BlueprintVarActionDetails_VariableNameLabel", "Variable Name"))
 		.ToolTip(VarNameTooltip)
 		.Font(IDetailLayoutBuilder::GetDetailFont())
 	]
@@ -462,7 +462,6 @@ void FBlueprintVarActionDetails::CustomizeDetails( IDetailLayoutBuilder& DetailL
 		UClass* VariableClass = (VariableProperty != NULL) ? VariableProperty->GetTypedOuter<UClass>() : NULL;
 
 		FText ErrorMessage;
-		const FString TooltipStr = LOCTEXT("DefaultValue_Tooltip", "Edit the default value of this variable.").ToString();
 		IDetailCategoryBuilder& DefaultValueCategory = DetailLayout.EditCategory(TEXT("DefaultValueCategory"), LOCTEXT("DefaultValueCategoryHeading", "Default Value").ToString());
 
 		if (VariableProperty == NULL)
@@ -505,8 +504,8 @@ void FBlueprintVarActionDetails::CustomizeDetails( IDetailLayoutBuilder& DetailL
 			DefaultValueCategory.AddCustomRow( TEXT("Error") )
 			[
 				SNew(STextBlock)
-				.ToolTipText(ErrorMessage.ToString())
-				.Text(ErrorMessage.ToString())
+				.ToolTipText(ErrorMessage)
+				.Text(ErrorMessage)
 				.Font(DetailFontInfo)
 			];
 		}
@@ -4164,7 +4163,7 @@ void FBlueprintComponentDetails::CustomizeDetails(IDetailLayoutBuilder& DetailLa
 				.AutoWidth()
 				[
 					SNew( STextBlock )
-						.Text( LOCTEXT("ScriptingEvents_Label", "Add Event").ToString()  )
+						.Text( LOCTEXT("ScriptingEvents_Label", "Add Event")  )
 						.Font( IDetailLayoutBuilder::GetDetailFont() )
 				]
 			]
@@ -4472,9 +4471,9 @@ FText FBlueprintComponentDetails::GetSocketName() const
 
 	if (CachedNodePtr->GetSCSNode() != NULL)
 	{
-		return FText::FromString(CachedNodePtr->GetSCSNode()->AttachToName.ToString());
+		return FText::FromName(CachedNodePtr->GetSCSNode()->AttachToName);
 	}
-	return FText();
+	return FText::GetEmpty();
 }
 
 void FBlueprintComponentDetails::OnBrowseSocket()
diff --git a/Engine/Source/Editor/Kismet/Private/BlueprintEditor.cpp b/Engine/Source/Editor/Kismet/Private/BlueprintEditor.cpp
--- a/Engine/Source/Editor/Kismet/Private/BlueprintEditor.cpp
+++ b/Engine/Source/Editor/Kismet/Private/BlueprintEditor.cpp
@@ -1133,7 +1133,7 @@ void FBlueprintEditor::PostRegenerateMenusAndToolbars()
 				SNew(STextBlock)
 				.ColorAndOpacity( FSlateColor::UseSubduedForeground() )
 				.ShadowOffset( FVector2D::UnitVector )
-				.Text(LOCTEXT("BlueprintEditor_ParentClass", "Parent class: ").ToString())
+				.Text(LOCTEXT("BlueprintEditor_ParentClass", "Parent class: "))
 			]
 			+SHorizontalBox::Slot()
 			.AutoWidth()
@@ -1148,8 +1148,8 @@ void FBlueprintEditor::PostRegenerateMenusAndToolbars()
 			[
 				SNew(STextBlock)
 				.ShadowOffset( FVector2D::UnitVector )
-				.Text(this, &FBlueprintEditor::GetParentClassName)
-				.ToolTipText(LOCTEXT("ParentClassToolTip", "The class that the current Blueprint is based on. The parent provides the base definition, which the current Blueprint extends.").ToString())
+				.Text(this, &FBlueprintEditor::GetParentClassNameText)
+				.ToolTipText(LOCTEXT("ParentClassToolTip", "The class that the current Blueprint is based on. The parent provides the base definition, which the current Blueprint extends."))
 				.Visibility(this, &FBlueprintEditor::GetParentClassNameVisibility)
 			]
 			+SHorizontalBox::Slot()
@@ -1161,7 +1161,7 @@ void FBlueprintEditor::PostRegenerateMenusAndToolbars()
 				.OnClicked( this, &FBlueprintEditor::OnFindParentClassInContentBrowserClicked )
 				.IsEnabled( this, &FBlueprintEditor::IsParentClassABlueprint )
 				.Visibility( this, &FBlueprintEditor::ParentClassButtonsVisibility )
-				.ToolTipText( LOCTEXT("FindParentInCBToolTip", "Find parent in Content Browser").ToString() )
+				.ToolTipText( LOCTEXT("FindParentInCBToolTip", "Find parent in Content Browser") )
 				.ContentPadding(4.0f)
 				.ForegroundColor( FSlateColor::UseForeground() )
 				[
@@ -1178,7 +1178,7 @@ void FBlueprintEditor::PostRegenerateMenusAndToolbars()
 				.OnClicked( this, &FBlueprintEditor::OnEditParentClassClicked )
 				.IsEnabled( this, &FBlueprintEditor::IsParentClassABlueprint )
 				.Visibility( this, &FBlueprintEditor::ParentClassButtonsVisibility )
-				.ToolTipText( LOCTEXT("EditParentClassToolTip", "Open parent in editor").ToString() )
+				.ToolTipText( LOCTEXT("EditParentClassToolTip", "Open parent in editor") )
 				.ContentPadding(4.0f)
 				.ForegroundColor( FSlateColor::UseForeground() )
 				[
@@ -1209,10 +1209,10 @@ void FBlueprintEditor::PostRegenerateMenusAndToolbars()
 	}
 }
 
-FString FBlueprintEditor::GetParentClassName() const
+FText FBlueprintEditor::GetParentClassNameText() const
 {
 	UClass* ParentClass = (GetBlueprintObj() != NULL) ? GetBlueprintObj()->ParentClass : NULL;
-	return ((ParentClass != NULL) ? FText::FromName( ParentClass->GetFName() ) : LOCTEXT("BlueprintEditor_NoParentClass", "None")).ToString();
+	return (ParentClass != NULL) ? FText::FromName( ParentClass->GetFName() ) : LOCTEXT("BlueprintEditor_NoParentClass", "None");
 }
 
 bool FBlueprintEditor::IsParentClassOfObjectABlueprint( const UBlueprint* Blueprint ) const
@@ -1297,7 +1297,7 @@ void FBlueprintEditor::OnEditParentClassNativeCodeClicked()
 FText FBlueprintEditor::GetTextForNativeParentClassHeaderLink() const
 {
 	// it could be done using FSourceCodeNavigation, but it could be slow
-	return FText::FromString(FString::Printf(TEXT("%s.h"), *GetParentClassName()));
+	return FText::FromString(FString::Printf(TEXT("%s.h"), *GetParentClassNameText().ToString()));
 }
 
 FReply FBlueprintEditor::OnFindParentClassInContentBrowserClicked()
@@ -5338,29 +5338,29 @@ void FBlueprintEditor::OnAddNewDelegate()
 
 void FBlueprintEditor::NewDocument_OnClicked(ECreatedDocumentType GraphType)
 {
-	FString DocumentNameString;
+	FText DocumentNameText;
 	bool bResetMyBlueprintFilter = false;
 
 	switch (GraphType)
 	{
 	case CGT_NewFunctionGraph:
-		DocumentNameString = LOCTEXT("NewDocFuncName", "NewFunction").ToString();
+		DocumentNameText = LOCTEXT("NewDocFuncName", "NewFunction");
 		bResetMyBlueprintFilter = true;
 		break;
 	case CGT_NewEventGraph:
-		DocumentNameString = LOCTEXT("NewDocEventGraphName", "NewEventGraph").ToString();
+		DocumentNameText = LOCTEXT("NewDocEventGraphName", "NewEventGraph");
 		bResetMyBlueprintFilter = true;
 		break;
 	case CGT_NewMacroGraph:
-		DocumentNameString = LOCTEXT("NewDocMacroName", "NewMacro").ToString();
+		DocumentNameText = LOCTEXT("NewDocMacroName", "NewMacro");
 		bResetMyBlueprintFilter = true;
 		break;
 	case CGT_NewAnimationGraph:
-		DocumentNameString = LOCTEXT("NewDocAnimationGraphName", "NewAnimationGraph").ToString();
+		DocumentNameText = LOCTEXT("NewDocAnimationGraphName", "NewAnimationGraph");
 		bResetMyBlueprintFilter = true;
 		break;
 	default:
-		DocumentNameString = LOCTEXT("NewDocNewName", "NewDocument").ToString();
+		DocumentNameText = LOCTEXT("NewDocNewName", "NewDocument");
 		break;
 	}
 	
@@ -5370,13 +5370,13 @@ void FBlueprintEditor::NewDocument_OnClicked(ECreatedDocumentType GraphType)
 		MyBlueprintWidget->OnResetItemFilter();
 	}
 
-	FName DocumentName = FName(*DocumentNameString);
+	FName DocumentName = FName(*DocumentNameText.ToString());
 
 	// Make sure the new name is valid
 	int32 Index = 0;
 	while (!FBlueprintEditorUtils::IsGraphNameUnique(GetBlueprintObj(), DocumentName))
 	{
-		DocumentName = FName(*FString::Printf(TEXT("%s%i"), *DocumentNameString, Index));
+		DocumentName = FName(*FString::Printf(TEXT("%s%i"), *DocumentNameText.ToString(), Index));
 		++Index;
 	}
 
diff --git a/Engine/Source/Editor/Kismet/Private/FindInBlueprints.cpp b/Engine/Source/Editor/Kismet/Private/FindInBlueprints.cpp
--- a/Engine/Source/Editor/Kismet/Private/FindInBlueprints.cpp
+++ b/Engine/Source/Editor/Kismet/Private/FindInBlueprints.cpp
@@ -254,7 +254,7 @@ void SFindInBlueprints::Construct( const FArguments& InArgs, TSharedPtr<FBluepri
 					.IsChecked(this, &SFindInBlueprints::OnGetFindModeChecked)
 					[
 						SNew(STextBlock)
-						.Text(LOCTEXT("BlueprintSearchModeChange", "Find In Current Blueprint Only").ToString())
+						.Text(LOCTEXT("BlueprintSearchModeChange", "Find In Current Blueprint Only"))
 					]
 				]
 			]
@@ -469,7 +469,7 @@ void SFindInBlueprints::MatchTokensWithinBlueprint(const TArray<FString>& Tokens
 	if( CommentNodes.Num() != 0 )
 	{
 		// Create the root for the comment hit results
-		const FString NodeName = FString::Printf(*LOCTEXT("BlueprintResultCommentResults", "Comments containing search term(s)").ToString() );
+		const FString NodeName = LOCTEXT("BlueprintResultCommentResults", "Comments containing search term(s)").ToString();
 		FSearchResult RootCommentResult(new FFindInBlueprintsResult(NodeName, RootSearchResult, CommentNodes[0]->GetClass(), 0));
 		
 		for(auto It(CommentNodes.CreateConstIterator());It;++It)
diff --git a/Engine/Source/Editor/Kismet/Private/SBlueprintActionMenu.cpp b/Engine/Source/Editor/Kismet/Private/SBlueprintActionMenu.cpp
--- a/Engine/Source/Editor/Kismet/Private/SBlueprintActionMenu.cpp
+++ b/Engine/Source/Editor/Kismet/Private/SBlueprintActionMenu.cpp
@@ -288,7 +288,7 @@ void SBlueprintActionMenu::Construct( const FArguments& InArgs, TSharedPtr<FBlue
 							TEXT("BlueprintActionMenuContextToggle")))
 						[
 							SNew(STextBlock)
-							.Text(LOCTEXT("BlueprintActionMenuContextToggle", "Context Sensitive").ToString())
+							.Text(LOCTEXT("BlueprintActionMenuContextToggle", "Context Sensitive"))
 						]
 					]
 				]
diff --git a/Engine/Source/Editor/Kismet/Private/SBlueprintDiff.cpp b/Engine/Source/Editor/Kismet/Private/SBlueprintDiff.cpp
--- a/Engine/Source/Editor/Kismet/Private/SBlueprintDiff.cpp
+++ b/Engine/Source/Editor/Kismet/Private/SBlueprintDiff.cpp
@@ -403,7 +403,7 @@ void SBlueprintDiff::Construct( const FArguments& InArgs)
 		.VAlign(VAlign_Center)
 		[
 			SNew(STextBlock)
-			.Text(LOCTEXT("BlueprintDifGraphsToolTip", "Select Graph to Diff").ToString())
+			.Text(LOCTEXT("BlueprintDifGraphsToolTip", "Select Graph to Diff"))
 		];
 
 	this->ChildSlot
@@ -441,7 +441,7 @@ void SBlueprintDiff::Construct( const FArguments& InArgs)
 										[
 											SNew(SImage)
 											.Image(this, &SBlueprintDiff::GetLockViewImage)
-											.ToolTipText(LOCTEXT("BPDifLock", "Lock the blueprint views?").ToString())
+											.ToolTipText(LOCTEXT("BPDifLock", "Lock the blueprint views?"))
 										]
 									]
 									+SVerticalBox::Slot()
@@ -458,11 +458,11 @@ void SBlueprintDiff::Construct( const FArguments& InArgs)
 												.BorderImage(FEditorStyle::GetBrush("PropertyWindow.CategoryBackground"))
 												.Padding(FMargin(2.0f))
 												.ForegroundColor(FEditorStyle::GetColor("PropertyWindow.CategoryForeground"))
-												.ToolTipText(LOCTEXT("BlueprintDifGraphsToolTip", "Select Graph to Diff").ToString())
+												.ToolTipText(LOCTEXT("BlueprintDifGraphsToolTip", "Select Graph to Diff"))
 												.HAlign(HAlign_Center)
 												[
 													SNew(STextBlock)
-													.Text(LOCTEXT("BlueprintDifGraphs", "Graphs").ToString())
+													.Text(LOCTEXT("BlueprintDifGraphs", "Graphs"))
 												]
 											]
 											+SVerticalBox::Slot()
@@ -491,7 +491,7 @@ void SBlueprintDiff::Construct( const FArguments& InArgs)
 										.Content()
 										[
 											SNew(STextBlock)
-											.Text(LOCTEXT("DifBlueprintDefaultsToolTip", "Diff Blueprint Defaults").ToString())
+											.Text(LOCTEXT("DifBlueprintDefaultsToolTip", "Diff Blueprint Defaults"))
 										]
 									]
 									+SVerticalBox::Slot()
@@ -660,7 +660,7 @@ void SBlueprintDiff::FDiffPanel::GeneratePanel(UEdGraph* Graph, UEdGraph* GraphT
 								.HAlign(HAlign_Center)
 								.VAlign(VAlign_Center)
 								[
-									SNew(STextBlock).Text( LOCTEXT("BPDifPanelNoGraphTip", "Graph does not exist in this revision").ToString())
+									SNew(STextBlock).Text( LOCTEXT("BPDifPanelNoGraphTip", "Graph does not exist in this revision"))
 								];
 	if(Graph)
 	{
diff --git a/Engine/Source/Editor/Kismet/Private/SBlueprintLibraryPalette.cpp b/Engine/Source/Editor/Kismet/Private/SBlueprintLibraryPalette.cpp
--- a/Engine/Source/Editor/Kismet/Private/SBlueprintLibraryPalette.cpp
+++ b/Engine/Source/Editor/Kismet/Private/SBlueprintLibraryPalette.cpp
@@ -371,7 +371,7 @@ TSharedRef<SWidget> SBlueprintLibraryPalette::ConstructClassFilterDropdownConten
 						.OnClicked(this, &SBlueprintLibraryPalette::ClearClassFilter)
 						.ToolTip(ClearFilterToolTip)
 					[
-						SNew(STextBlock).Text(LOCTEXT("All", "All").ToString())
+						SNew(STextBlock).Text(LOCTEXT("All", "All"))
 					]
 				]
 
diff --git a/Engine/Source/Editor/Kismet/Private/SKismetInspector.cpp b/Engine/Source/Editor/Kismet/Private/SKismetInspector.cpp
--- a/Engine/Source/Editor/Kismet/Private/SKismetInspector.cpp
+++ b/Engine/Source/Editor/Kismet/Private/SKismetInspector.cpp
@@ -49,7 +49,7 @@ TSharedRef<SWidget> SKismetInspector::MakeContextualEditingWidget(struct FKismet
 			.Padding( 2.0f, 14.0f, 2.0f, 2.0f )
 			[
 				SNew( STextBlock )
-				.Text( LOCTEXT("NoNodesSelected", "Select a node to edit details.").ToString() )
+				.Text( LOCTEXT("NoNodesSelected", "Select a node to edit details.") )
 			];
 		}
 		else
@@ -89,11 +89,11 @@ TSharedRef<SWidget> SKismetInspector::MakeContextualEditingWidget(struct FKismet
 				.VAlign( VAlign_Top )
 				[
 					SNew(SCheckBox)
-					.ToolTipText(LOCTEXT("TogglePublicView", "Toggle Public View").ToString())
+					.ToolTipText(LOCTEXT("TogglePublicView", "Toggle Public View"))
 					.IsChecked(this, &SKismetInspector::GetPublicViewCheckboxState)
 					.OnCheckStateChanged( this, &SKismetInspector::SetPublicViewCheckboxState)
 					[
-						SNew(STextBlock) .Text(LOCTEXT("PublicViewCheckboxLabel", "Public View").ToString())
+						SNew(STextBlock) .Text(LOCTEXT("PublicViewCheckboxLabel", "Public View"))
 					]
 				];
 		}
diff --git a/Engine/Source/Editor/Kismet/Private/SSCSEditor.cpp b/Engine/Source/Editor/Kismet/Private/SSCSEditor.cpp
--- a/Engine/Source/Editor/Kismet/Private/SSCSEditor.cpp
+++ b/Engine/Source/Editor/Kismet/Private/SSCSEditor.cpp
@@ -1645,8 +1645,8 @@ FText SSCS_RowWidget::GetTooltipText() const
 		}
 		else
 		{
-			return LOCTEXT("DefaultSceneRootToolTip", "This is the default scene root component. It cannot be renamed or deleted. Adding a new scene component will automatically replace it as the new root.");
-		}
+		return LOCTEXT("DefaultSceneRootToolTip", "This is the default scene root component. It cannot be renamed or deleted. Adding a new scene component will automatically replace it as the new root.");
+	}
 	}
 	else
 	{
@@ -1796,25 +1796,25 @@ void SSCSEditor::Construct( const FArguments& InArgs, TSharedPtr<FBlueprintEdito
 		[
 			SNew(SVerticalBox)
 
-			+ SVerticalBox::Slot()
+			+SVerticalBox::Slot()
 			.AutoHeight()
 			[
 				SAssignNew(ComponentClassCombo, SComponentClassCombo)
 				.OnComponentClassSelected(this, &SSCSEditor::PerformComboAddClass)
 			]
 
-			+ SVerticalBox::Slot()
+			+SVerticalBox::Slot()
 			.AutoHeight()
 			[
 				SNew(SSpacer)
 				.Size(FVector2D(0.0f,2.0f))
 			]
 
-			+ SVerticalBox::Slot()
+			+SVerticalBox::Slot()
 			.FillHeight(.8f)
 			[
 				SAssignNew(SCSTreeWidget, SSCSTreeType)
-				.ToolTipText(LOCTEXT("DropAssetToAddComponent", "Drop asset here to add component.").ToString())
+				.ToolTipText(LOCTEXT("DropAssetToAddComponent", "Drop asset here to add component."))
 				.SCSEditor( this )
 				.TreeItemsSource( &RootNodes )
 				.SelectionMode(ESelectionMode::Multi)
@@ -1828,13 +1828,13 @@ void SSCSEditor::Construct( const FArguments& InArgs, TSharedPtr<FBlueprintEdito
 				(
 					SNew(SHeaderRow)
 
-					+ SHeaderRow::Column(SCS_ColumnName_Mobility)
-						.DefaultLabel(LOCTEXT("MobilityColumnLabel", "Mobility").ToString())
+					+SHeaderRow::Column(SCS_ColumnName_Mobility)
+						.DefaultLabel(LOCTEXT("MobilityColumnLabel", "Mobility"))
 						.FixedWidth(16.0f) // mobility icons are 16px (16 slate-units = 16px, when application scale == 1)
 						.HeaderContent()
 						[
 							SNew(SHorizontalBox)
-								.ToolTip(SNew(SToolTip).Text(LOCTEXT("MobilityColumnTooltip", "Mobility").ToString()))
+								.ToolTip(SNew(SToolTip).Text(LOCTEXT("MobilityColumnTooltip", "Mobility")))
 							+SHorizontalBox::Slot()
 								.FillWidth(1.0f)
 								.VAlign(VAlign_Center)
@@ -1843,13 +1843,13 @@ void SSCSEditor::Construct( const FArguments& InArgs, TSharedPtr<FBlueprintEdito
 									SNew(SImage).Image(MobilityHeaderBrush)
 								]
 						]
-					
+				
 					+ SHeaderRow::Column(SCS_ColumnName_ComponentClass)
-					.DefaultLabel( LOCTEXT("Class", "Class").ToString() )
+					.DefaultLabel( LOCTEXT("Class", "Class") )
 					.FillWidth(4)
 					
 					+ SHeaderRow::Column(SCS_ColumnName_Asset)
-					.DefaultLabel( LOCTEXT("Asset", "Asset").ToString() )
+					.DefaultLabel( LOCTEXT("Asset", "Asset") )
 					.FillWidth(3)
 				)
 			]
diff --git a/Engine/Source/Editor/Kismet/Private/STimelineEditor.cpp b/Engine/Source/Editor/Kismet/Private/STimelineEditor.cpp
--- a/Engine/Source/Editor/Kismet/Private/STimelineEditor.cpp
+++ b/Engine/Source/Editor/Kismet/Private/STimelineEditor.cpp
@@ -73,7 +73,7 @@ void SCurveAssetWidget::Construct(const FArguments& InArgs)
 								.Padding(4)
 								[
 									SNew(STextBlock)
-									.Text( LOCTEXT( "AssetPath", "Path:" ).ToString() )
+									.Text( LOCTEXT( "AssetPath", "Path:" ) )
 								]
 
 							//Editable text
@@ -100,7 +100,7 @@ void SCurveAssetWidget::Construct(const FArguments& InArgs)
 							[
 								SNew(SButton)
 								.HAlign(HAlign_Center)
-								.Text( LOCTEXT( "CreateAssetButton", "Create External Curve" ).ToString() )
+								.Text( LOCTEXT( "CreateAssetButton", "Create External Curve" ) )
 								.OnClicked(this, &SCurveAssetWidget::OnCreateButtonClicked)
 							]
 						]
@@ -248,7 +248,7 @@ void STimelineEdTrack::Construct(const FArguments& InArgs, TSharedPtr<FTimelineE
 							.ButtonStyle( FEditorStyle::Get(), "NoBorder" )
 							.OnClicked(this, &STimelineEdTrack::OnClickUse)
 							.ContentPadding(1.f)
-							.ToolTipText(NSLOCTEXT("TimelineEdTrack", "TimelineEdTrack_Use", "Use External Curve").ToString())
+							.ToolTipText(NSLOCTEXT("TimelineEdTrack", "TimelineEdTrack_Use", "Use External Curve"))
 							[
 								SNew(SImage)
 								.Image( FEditorStyle::GetBrush(TEXT("PropertyWindow.Button_Use")) )
@@ -264,7 +264,7 @@ void STimelineEdTrack::Construct(const FArguments& InArgs, TSharedPtr<FTimelineE
 							.ButtonStyle( FEditorStyle::Get(), "NoBorder" )
 							.OnClicked(this, &STimelineEdTrack::OnClickBrowse)
 							.ContentPadding(0)
-							.ToolTipText(NSLOCTEXT("TimelineEdTrack", "TimelineEdTrack_Browse", "Browse External Curve").ToString())
+							.ToolTipText(NSLOCTEXT("TimelineEdTrack", "TimelineEdTrack_Browse", "Browse External Curve"))
 							[
 								SNew(SImage)
 								.Image( FEditorStyle::GetBrush(TEXT("PropertyWindow.Button_Browse")) )
@@ -280,7 +280,7 @@ void STimelineEdTrack::Construct(const FArguments& InArgs, TSharedPtr<FTimelineE
 							.ButtonStyle( FEditorStyle::Get(), "NoBorder" )
 							.OnClicked(this, &STimelineEdTrack::OnClickClear)
 							.ContentPadding(1.f)
-							.ToolTipText(NSLOCTEXT("TimelineEdTrack", "TimelineEdTrack_Clear", "Convert to Internal Curve").ToString())
+							.ToolTipText(NSLOCTEXT("TimelineEdTrack", "TimelineEdTrack_Clear", "Convert to Internal Curve"))
 							[
 								SNew(SImage)
 								.Image( FEditorStyle::GetBrush(TEXT("PropertyWindow.Button_Clear")) )
@@ -845,7 +845,7 @@ void STimelineEditor::Construct(const FArguments& InArgs, TSharedPtr<FBlueprintE
 					SNew(SImage)
 					.Image(FEditorStyle::GetBrush("TimelineEditor.AddFloatTrack"))
 				]
-				.ToolTipText( LOCTEXT( "AddFloatTrack", "Add Float Track" ).ToString() )
+				.ToolTipText( LOCTEXT( "AddFloatTrack", "Add Float Track" ) )
 				.OnClicked( this, &STimelineEditor::CreateNewTrack, FTimelineEdTrack::TT_FloatInterp )
 			]
 			+SHorizontalBox::Slot()
@@ -859,7 +859,7 @@ void STimelineEditor::Construct(const FArguments& InArgs, TSharedPtr<FBlueprintE
 					SNew(SImage)
 					.Image(FEditorStyle::GetBrush("TimelineEditor.AddVectorTrack"))
 				]
-				.ToolTipText( LOCTEXT( "AddVectorTrack", "Add Vector Track" ).ToString() )
+				.ToolTipText( LOCTEXT( "AddVectorTrack", "Add Vector Track" ) )
 				.OnClicked( this, &STimelineEditor::CreateNewTrack, FTimelineEdTrack::TT_VectorInterp )
 			]
 			+SHorizontalBox::Slot()
@@ -873,7 +873,7 @@ void STimelineEditor::Construct(const FArguments& InArgs, TSharedPtr<FBlueprintE
 					SNew(SImage)
 					.Image(FEditorStyle::GetBrush("TimelineEditor.AddEventTrack"))
 				]
-				.ToolTipText( LOCTEXT( "AddEventTrack", "Add Event Track" ).ToString() )
+				.ToolTipText( LOCTEXT( "AddEventTrack", "Add Event Track" ) )
 				.OnClicked( this, &STimelineEditor::CreateNewTrack, FTimelineEdTrack::TT_Event )
 			]
 			+SHorizontalBox::Slot()
@@ -887,7 +887,7 @@ void STimelineEditor::Construct(const FArguments& InArgs, TSharedPtr<FBlueprintE
 					SNew(SImage)
 					.Image(FEditorStyle::GetBrush("TimelineEditor.AddColorTrack"))
 				]
-				.ToolTipText( LOCTEXT( "AddColorTrack", "Add Color Track" ).ToString() )
+				.ToolTipText( LOCTEXT( "AddColorTrack", "Add Color Track" ) )
 				.OnClicked( this, &STimelineEditor::CreateNewTrack, FTimelineEdTrack::TT_LinearColorInterp )
 			]
 			+SHorizontalBox::Slot()
@@ -901,7 +901,7 @@ void STimelineEditor::Construct(const FArguments& InArgs, TSharedPtr<FBlueprintE
 					SNew(SImage)
 					.Image(FEditorStyle::GetBrush("TimelineEditor.AddCurveAssetTrack"))
 				]
-				.ToolTipText( LOCTEXT( "AddExternalAsset", "Add Selected Curve Asset" ).ToString() )
+				.ToolTipText( LOCTEXT( "AddExternalAsset", "Add Selected Curve Asset" ) )
 				.IsEnabled( this, &STimelineEditor::IsCurveAssetSelected )
 				.OnClicked( this, &STimelineEditor::CreateNewTrackFromAsset )
 			]
@@ -912,7 +912,7 @@ void STimelineEditor::Construct(const FArguments& InArgs, TSharedPtr<FBlueprintE
 			[
 				// Length label
 				SNew(STextBlock)
-				.Text( LOCTEXT( "Length", "Length" ).ToString() )
+				.Text( LOCTEXT( "Length", "Length" ) )
 			]
 			+SHorizontalBox::Slot()
 			.AutoWidth()
@@ -934,7 +934,7 @@ void STimelineEditor::Construct(const FArguments& InArgs, TSharedPtr<FBlueprintE
 				.IsChecked( this, &STimelineEditor::IsUseLastKeyframeChecked )
 				.OnCheckStateChanged( this, &STimelineEditor::OnUseLastKeyframeChanged )
 				[
-					SNew(STextBlock) .Text( LOCTEXT( "UseLastKeyframe", "Use Last Keyframe?" ).ToString() )
+					SNew(STextBlock) .Text( LOCTEXT( "UseLastKeyframe", "Use Last Keyframe?" ) )
 				]
 			]
 			+SHorizontalBox::Slot()
@@ -946,7 +946,7 @@ void STimelineEditor::Construct(const FArguments& InArgs, TSharedPtr<FBlueprintE
 				.IsChecked( this, &STimelineEditor::IsAutoPlayChecked )
 				.OnCheckStateChanged( this, &STimelineEditor::OnAutoPlayChanged )
 				[
-					SNew(STextBlock) .Text( LOCTEXT( "AutoPlay", "AutoPlay" ).ToString() )
+					SNew(STextBlock) .Text( LOCTEXT( "AutoPlay", "AutoPlay" ) )
 				]
 			]
 			+SHorizontalBox::Slot()
@@ -958,7 +958,7 @@ void STimelineEditor::Construct(const FArguments& InArgs, TSharedPtr<FBlueprintE
 				.IsChecked( this, &STimelineEditor::IsLoopChecked )
 				.OnCheckStateChanged( this, &STimelineEditor::OnLoopChanged )
 				[
-					SNew(STextBlock) .Text( LOCTEXT( "Loop", "Loop" ).ToString() )
+					SNew(STextBlock) .Text( LOCTEXT( "Loop", "Loop" ) )
 				]
 			]
 			+SHorizontalBox::Slot()
@@ -970,7 +970,7 @@ void STimelineEditor::Construct(const FArguments& InArgs, TSharedPtr<FBlueprintE
 				.IsChecked( this, &STimelineEditor::IsReplicatedChecked )
 				.OnCheckStateChanged( this, &STimelineEditor::OnReplicatedChanged )
 				[
-					SNew(STextBlock) .Text( LOCTEXT( "Replicated", "Replicated" ).ToString() )
+					SNew(STextBlock) .Text( LOCTEXT( "Replicated", "Replicated" ) )
 				]
 			]
 		]
@@ -996,15 +996,15 @@ void STimelineEditor::Construct(const FArguments& InArgs, TSharedPtr<FBlueprintE
 }
 END_SLATE_FUNCTION_BUILD_OPTIMIZATION
 
-FString STimelineEditor::GetTimelineName() const
+FText STimelineEditor::GetTimelineName() const
 {
 	if(TimelineObj != NULL)
 	{
-		return UTimelineTemplate::TimelineTemplateNameToVariableName(TimelineObj->GetFName());
+		return FText::FromString(UTimelineTemplate::TimelineTemplateNameToVariableName(TimelineObj->GetFName()));
 	}
 	else
 	{
-		return LOCTEXT( "NoTimeline", "No Timeline" ).ToString();
+		return LOCTEXT( "NoTimeline", "No Timeline" );
 	}
 }
 
diff --git a/Engine/Source/Editor/Kismet/Private/STimelineEditor.h b/Engine/Source/Editor/Kismet/Private/STimelineEditor.h
--- a/Engine/Source/Editor/Kismet/Private/STimelineEditor.h
+++ b/Engine/Source/Editor/Kismet/Private/STimelineEditor.h
@@ -277,7 +277,7 @@ private:
 	void OnDeleteSelectedTracks();
 
 	/** Get the name of the timeline we are editing */
-	FString GetTimelineName() const;
+	FText GetTimelineName() const;
 
 	/** Get state of autoplay box*/
 	ESlateCheckBoxState::Type IsAutoPlayChecked() const;
diff --git a/Engine/Source/Editor/Kismet/Private/UserDefinedEnumEditor.cpp b/Engine/Source/Editor/Kismet/Private/UserDefinedEnumEditor.cpp
--- a/Engine/Source/Editor/Kismet/Private/UserDefinedEnumEditor.cpp
+++ b/Engine/Source/Editor/Kismet/Private/UserDefinedEnumEditor.cpp
@@ -155,7 +155,7 @@ void FEnumDetails::CustomizeDetails( IDetailLayoutBuilder& DetailLayout )
 				.HAlign(HAlign_Right)
 				[
 					SNew(SButton)
-					.Text(LOCTEXT("FunctionNewInputArg", "New").ToString())
+					.Text(LOCTEXT("FunctionNewInputArg", "New"))
 					.OnClicked(this, &FEnumDetails::OnAddNewEnumerator)
 				]
 			];
@@ -346,15 +346,15 @@ bool FUserDefinedEnumIndexLayout::IsValidEnumeratorDisplayName(const FText& NewT
 
 	const FString NewName = NewText.ToString();
 	bool bUnchangedName = (NewName == FEnumEditorUtils::GetEnumeratorDisplayName(TargetEnum.Get(), EnumeratorIndex));
-	FString ErrorMsg;
+	FText ErrorMsg;
 	if (NewText.IsEmpty())
 	{
-		ErrorMsg = LOCTEXT("NameMissingError", "You must provide a name.").ToString();
+		ErrorMsg = LOCTEXT("NameMissingError", "You must provide a name.");
 		bValidName = false;
 	}
 	else if (!bUnchangedName && !FEnumEditorUtils::IsEnumeratorDisplayNameValid(TargetEnum.Get(), NewName))
 	{
-		ErrorMsg = FText::Format(LOCTEXT("NameInUseError", "'{0}' is already in use."), FText::FromString( NewName)).ToString();
+		ErrorMsg = FText::Format(LOCTEXT("NameInUseError", "'{0}' is already in use."), NewText);
 		bValidName = false;
 	}
 
diff --git a/Engine/Source/Editor/Kismet/Public/BlueprintEditor.h b/Engine/Source/Editor/Kismet/Public/BlueprintEditor.h
--- a/Engine/Source/Editor/Kismet/Public/BlueprintEditor.h
+++ b/Engine/Source/Editor/Kismet/Public/BlueprintEditor.h
@@ -466,7 +466,7 @@ protected:
 	virtual void PostRegenerateMenusAndToolbars() OVERRIDE;
 
 	/** Returns the name of the Blueprint's parent class */
-	FString GetParentClassName() const;
+	FText GetParentClassNameText() const;
 
 	/** Handler for "Find parent class in CB" button */
 	FReply OnFindParentClassInContentBrowserClicked();
diff --git a/Engine/Source/Editor/Layers/Private/LayerViewModel.cpp b/Engine/Source/Editor/Layers/Private/LayerViewModel.cpp
--- a/Engine/Source/Editor/Layers/Private/LayerViewModel.cpp
+++ b/Engine/Source/Editor/Layers/Private/LayerViewModel.cpp
@@ -137,14 +137,17 @@ void FLayerViewModel::RenameTo( const FName& NewLayerName )
 }
 
 
-bool FLayerViewModel::CanAssignActors( const TArray< TWeakObjectPtr<AActor> > Actors, FString& OutMessage ) const
+bool FLayerViewModel::CanAssignActors( const TArray< TWeakObjectPtr<AActor> > Actors, FText& OutMessage ) const
 {
 	if( !Layer.IsValid() )
 	{
-		OutMessage = LOCTEXT("InvalidLayer", "Invalid Layer").ToString();
+		OutMessage = LOCTEXT("InvalidLayer", "Invalid Layer");
 		return false;
 	}
 
+	FFormatNamedArguments LayerArgs;
+	LayerArgs.Add(TEXT("LayerName"), FText::FromName(Layer->LayerName));
+
 	bool bHasValidActorToAssign = false;
 	int32 bAlreadyAssignedActors = 0;
 	for( auto ActorIt = Actors.CreateConstIterator(); ActorIt; ++ActorIt )
@@ -152,13 +155,16 @@ bool FLayerViewModel::CanAssignActors( const TArray< TWeakObjectPtr<AActor> > Ac
 		TWeakObjectPtr< AActor > Actor = *ActorIt;
 		if( !Actor.IsValid() )
 		{
-			OutMessage = FString::Printf( *LOCTEXT("InvalidActors", "Cannot add invalid Actors to %s").ToString(), *Layer->LayerName.ToString() );
+			OutMessage = FText::Format(LOCTEXT("InvalidActors", "Cannot add invalid Actors to {LayerName}"), LayerArgs);
 			return false;
 		}
 
+		FFormatNamedArguments ActorArgs;
+		ActorArgs.Add(TEXT("ActorName"), FText::FromName(Actor->GetFName()));
+
 		if( !WorldLayers->IsActorValidForLayer( Actor ) )
 		{
-			OutMessage = FString::Printf( *LOCTEXT("InvalidLayers", "Actor '%s' cannot be associated with Layers").ToString(), *Actor->GetFName().ToString() );
+			OutMessage = FText::Format(LOCTEXT("InvalidLayers", "Actor '{ActorName}' cannot be associated with Layers"), ActorArgs);
 			return false;
 		}
 
@@ -174,48 +180,53 @@ bool FLayerViewModel::CanAssignActors( const TArray< TWeakObjectPtr<AActor> > Ac
 
 	if( bAlreadyAssignedActors == Actors.Num() )
 	{
-		OutMessage = FString::Printf( *LOCTEXT("AlreadyAssignedActors", "All Actors already assigned to %s").ToString(), *Layer->LayerName.ToString() );
+		OutMessage = FText::Format(LOCTEXT("AlreadyAssignedActors", "All Actors already assigned to {LayerName}"), LayerArgs);
 		return false;
 	}
 
 	if( bHasValidActorToAssign )
 	{
-		OutMessage = FString::Printf( *LOCTEXT("AssignActors", "Assign Actors to %s").ToString(), *Layer->LayerName.ToString() );
+		OutMessage = FText::Format(LOCTEXT("AssignActors", "Assign Actors to {LayerName}"), LayerArgs);
 		return true;
 	}
 
-	OutMessage.Empty();
+	OutMessage = FText::GetEmpty();
 	return false;
 }
 
 
-bool FLayerViewModel::CanAssignActor( const TWeakObjectPtr<AActor> Actor, FString& OutMessage ) const
+bool FLayerViewModel::CanAssignActor( const TWeakObjectPtr<AActor> Actor, FText& OutMessage ) const
 {
 	if( !Layer.IsValid() )
 	{
-		OutMessage = LOCTEXT("InvalidLayer", "Invalid Layer").ToString();
+		OutMessage = LOCTEXT("InvalidLayer", "Invalid Layer");
 		return false;
 	}
 
+	FFormatNamedArguments Args;
+	Args.Add(TEXT("LayerName"), FText::FromName(Layer->LayerName));
+
 	if( !Actor.IsValid() )
 	{
-		OutMessage = FString::Printf( *LOCTEXT("InvalidActor", "Cannot add invalid Actors to %s").ToString(), *Layer->LayerName.ToString() );
+		OutMessage = FText::Format(LOCTEXT("InvalidActor", "Cannot add invalid Actors to {LayerName}"), Args);
 		return false;
 	}
 
+	Args.Add(TEXT("ActorName"), FText::FromName(Actor->GetFName()));
+
 	if( !WorldLayers->IsActorValidForLayer( Actor ) )
 	{
-		OutMessage = FString::Printf(* LOCTEXT("InvalidLayers", "Actor '%s' cannot be associated with Layers").ToString(), *Actor->GetFName().ToString() );
+		OutMessage = FText::Format(LOCTEXT("InvalidLayers", "Actor '{ActorName}' cannot be associated with Layers"), Args);
 		return false;
 	}
 
 	if( Actor->Layers.Contains( Layer->LayerName )  )
 	{
-		OutMessage = FString::Printf( *LOCTEXT("AlreadyAssignedActor", "Already assigned to %s").ToString(), *Layer->LayerName.ToString() );
+		OutMessage = FText::Format(LOCTEXT("AlreadyAssignedActor", "Already assigned to {LayerName}"), Args);
 		return false;
 	}
 
-	OutMessage = FString::Printf( *LOCTEXT("AssignActor", "Assign to %s").ToString(), *Layer->LayerName.ToString() );
+	OutMessage = FText::Format(LOCTEXT("AssignActor", "Assign to {LayerName}"), Args);
 	return true;
 }
 
diff --git a/Engine/Source/Editor/Layers/Private/LayerViewModel.h b/Engine/Source/Editor/Layers/Private/LayerViewModel.h
--- a/Engine/Source/Editor/Layers/Private/LayerViewModel.h
+++ b/Engine/Source/Editor/Layers/Private/LayerViewModel.h
@@ -72,7 +72,7 @@ public:
 	 *	@return						if true then at least one actor can be assigned; 
 	 *								If false, either Invalid actors were discovered or all actors are already assigned
 	 */
-	bool CanAssignActors( const TArray< TWeakObjectPtr<AActor> > Actors, FString& OutMessage ) const;
+	bool CanAssignActors( const TArray< TWeakObjectPtr<AActor> > Actors, FText& OutMessage ) const;
 
 	/**
 	 *	Returns whether the specified actor can be assigned to the Layer
@@ -82,7 +82,7 @@ public:
 	 *	@return						if true then at least one actor can be assigned; 
 	 *								If false, either the Actor is Invalid or already assigned
 	 */
-	bool CanAssignActor( const TWeakObjectPtr<AActor> Actor, FString& OutMessage ) const;
+	bool CanAssignActor( const TWeakObjectPtr<AActor> Actor, FText& OutMessage ) const;
 	
 	/**
 	 *	Appends all of the actors associated with this layer to the specified list
diff --git a/Engine/Source/Editor/Layers/Private/SActorLayerCloud.h b/Engine/Source/Editor/Layers/Private/SActorLayerCloud.h
--- a/Engine/Source/Editor/Layers/Private/SActorLayerCloud.h
+++ b/Engine/Source/Editor/Layers/Private/SActorLayerCloud.h
@@ -44,7 +44,7 @@ public:
 			.HAlign( HAlign_Center )
 			[
 				SNew( STextBlock )
-				.Text( LOCTEXT("NoLayersMessage", "No Layers").ToString() )
+				.Text( LOCTEXT("NoLayersMessage", "No Layers") )
 				.Visibility( this, &SActorLayerCloud::GetNoLayersMessageVisibility )
 			]
 		];
@@ -95,7 +95,7 @@ private:
 								.VAlign( VAlign_Center )
 								.OnClicked( this, &SActorLayerCloud::OnRemoveLayerFromActorsClicked, Layer )
 								.ButtonStyle( FEditorStyle::Get(), "NoBorder" )
-								.ToolTipText( LOCTEXT("RemoveLayerButton", "Remove from Layer").ToString() )
+								.ToolTipText( LOCTEXT("RemoveLayerButton", "Remove from Layer") )
 								.Content()
 								[
 									SNew( SImage )
diff --git a/Engine/Source/Editor/Layers/Private/SLayerBrowser.h b/Engine/Source/Editor/Layers/Private/SLayerBrowser.h
--- a/Engine/Source/Editor/Layers/Private/SLayerBrowser.h
+++ b/Engine/Source/Editor/Layers/Private/SLayerBrowser.h
@@ -72,7 +72,7 @@ public:
 			.AutoHeight()
 			[
 				SNew( SSearchBox )
-				.ToolTipText( LOCTEXT("FilterSearchToolTip", "Type here to search layers").ToString() )
+				.ToolTipText( LOCTEXT("FilterSearchToolTip", "Type here to search layers") )
 				.HintText( LOCTEXT( "FilterSearchHint", "Search Layers" ) )
 				.OnTextChanged( SearchBoxLayerFilter.Get(), &LayerTextFilter::SetRawFilterText )
 			]
@@ -126,7 +126,7 @@ public:
 					.VAlign( VAlign_Center )
 					[
 						SNew( STextBlock )
-						.Text( LOCTEXT("ContentsLabel", "See Contents").ToString() )
+						.Text( LOCTEXT("ContentsLabel", "See Contents") )
 						.Visibility( this, &SLayerBrowser::IsVisibleIfModeIs, ELayerBrowserMode::Layers )
 						.ColorAndOpacity( this, &SLayerBrowser::GetInvertedForegroundIfHovered )
 					]
@@ -251,7 +251,7 @@ protected:
 		}
 
 		bool bCanAssign = false;
-		FString Message;
+		FText Message;
 		if( DragActorOp->Actors.Num() > 1 )
 		{
 			bCanAssign = SelectedLayerViewModel->CanAssignActors( DragActorOp->Actors, OUT Message );
@@ -263,11 +263,11 @@ protected:
 
 		if ( bCanAssign )
 		{
-			DragActorOp->SetToolTip( FActorDragDropGraphEdOp::ToolTip_CompatibleGeneric, FText::FromString(Message) );
+			DragActorOp->SetToolTip( FActorDragDropGraphEdOp::ToolTip_CompatibleGeneric, Message );
 		}
 		else
 		{
-			DragActorOp->SetToolTip( FActorDragDropGraphEdOp::ToolTip_IncompatibleGeneric, FText::FromString(Message) );
+			DragActorOp->SetToolTip( FActorDragDropGraphEdOp::ToolTip_IncompatibleGeneric, Message );
 		}
 
 		// We leave the event unhandled so the children of the ListView get a chance to grab the drag/drop
@@ -333,9 +333,9 @@ private:
 		return ( Mode == ELayerBrowserMode::Layers ) ? FEditorStyle::GetBrush( "LayerBrowser.ExploreLayerContents" ) : FEditorStyle::GetBrush( "LayerBrowser.ReturnToLayersList" );
 	}
 
-	FString GetLayerContentsHeaderText() const
+	FText GetLayerContentsHeaderText() const
 	{
-		return FString::Printf( *LOCTEXT("SelectedContentsLabel", "%s Contents").ToString(), *SelectedLayerViewModel->GetName() );
+		return FText::Format( LOCTEXT("SelectedContentsLabel", "{0} Contents"), FText::FromString(SelectedLayerViewModel->GetName()));
 	}
 
 	/**	Returns the visibility of the See Contents label */
diff --git a/Engine/Source/Editor/Layers/Private/SLayersViewRow.h b/Engine/Source/Editor/Layers/Private/SLayersViewRow.h
--- a/Engine/Source/Editor/Layers/Private/SLayersViewRow.h
+++ b/Engine/Source/Editor/Layers/Private/SLayersViewRow.h
@@ -169,7 +169,7 @@ protected:
 		}
 
 		bool bCanAssign = false;
-		FString Message;
+		FText Message;
 		if( DragActorOp->Actors.Num() > 1 )
 		{
 			bCanAssign = ViewModel->CanAssignActors( DragActorOp->Actors, OUT Message );
@@ -181,11 +181,11 @@ protected:
 
 		if ( bCanAssign )
 		{
-			DragActorOp->SetToolTip( FActorDragDropGraphEdOp::ToolTip_CompatibleGeneric, FText::FromString(Message) );
+			DragActorOp->SetToolTip( FActorDragDropGraphEdOp::ToolTip_CompatibleGeneric, Message );
 		}
 		else
 		{
-			DragActorOp->SetToolTip( FActorDragDropGraphEdOp::ToolTip_IncompatibleGeneric, FText::FromString(Message) );
+			DragActorOp->SetToolTip( FActorDragDropGraphEdOp::ToolTip_IncompatibleGeneric, Message );
 		}
 
 		return FReply::Handled();
@@ -246,7 +246,7 @@ private:
 		{
 			TSharedPtr<FActorDragDropGraphEdOp> DragDropActorOp = StaticCastSharedPtr<FActorDragDropGraphEdOp>(DragDropOp);
 
-			FString Message;
+			FText Message;
 			bCanAcceptDrop = ViewModel->CanAssignActors( DragDropActorOp->Actors, OUT Message );
 		}
 
