Fix for crash when exiting UFE with a running launch happening
TTP323997
#ue4
#ufe

diff --git a/Engine/Source/Developer/SessionLauncher/Private/Widgets/SSessionLauncher.cpp b/Engine/Source/Developer/SessionLauncher/Private/Widgets/SSessionLauncher.cpp
--- a/Engine/Source/Developer/SessionLauncher/Private/Widgets/SSessionLauncher.cpp
+++ b/Engine/Source/Developer/SessionLauncher/Private/Widgets/SSessionLauncher.cpp
@@ -14,6 +14,11 @@
 
 SSessionLauncher::~SSessionLauncher( )
 {
+	if (LauncherWorker.IsValid())
+	{
+		LauncherWorker->Cancel();
+		FPlatformProcess::Sleep(0.5f);
+	}
 	Model->OnProfileSelected().RemoveAll(this);
 }
 
