Fix for invalid error message when UFE completes a launch.
TTP334857
#ue4
#ufe

diff --git a/Engine/Source/Developer/SessionLauncher/Private/Widgets/SSessionLauncherProgress.h b/Engine/Source/Developer/SessionLauncher/Private/Widgets/SSessionLauncherProgress.h
--- a/Engine/Source/Developer/SessionLauncher/Private/Widgets/SSessionLauncherProgress.h
+++ b/Engine/Source/Developer/SessionLauncher/Private/Widgets/SSessionLauncherProgress.h
@@ -293,7 +293,7 @@ private:
 
 			FFormatNamedArguments Args;
 			Args.Add(TEXT("NumCompleted"), NumCompleted);
-			Args.Add(TEXT("NumFailed"), NumCompleted);
+			Args.Add(TEXT("NumFailed"), NumFailed);
 			Args.Add(TEXT("NumCanceled"), NumCanceled);
 
 			return FText::Format(LOCTEXT("TasksFinishedFormatText", "Operation finished. Completed: {NumCompleted}, Failed: {NumFailed}, Canceled: {NumCanceled}"), Args);
