GameLauncher: Fixed assertion in device list (TTP# 338166)

diff --git a/Engine/Source/Developer/SessionLauncher/Private/Widgets/Deploy/SSessionLauncherDeviceListRow.h b/Engine/Source/Developer/SessionLauncher/Private/Widgets/Deploy/SSessionLauncherDeviceListRow.h
--- a/Engine/Source/Developer/SessionLauncher/Private/Widgets/Deploy/SSessionLauncherDeviceListRow.h
+++ b/Engine/Source/Developer/SessionLauncher/Private/Widgets/Deploy/SSessionLauncherDeviceListRow.h
@@ -101,7 +101,7 @@ public:
 				.VAlign(VAlign_Center)
 				[
 					SNew(STextBlock)
-					.Text(this, &SSessionLauncherDeviceListRow::HandleHostPlatformText)
+						.Text(this, &SSessionLauncherDeviceListRow::HandleHostPlatformText)
 				];
 		}
 		else if (ColumnName == "Host")
@@ -172,7 +172,11 @@ private:
 	const FSlateBrush* HandleDeviceImage( ) const
 	{
 		const PlatformInfo::FPlatformInfo* const PlatformInfo = PlatformInfo::FindPlatformInfo(*DeviceProxy->GetPlatformName());
-		check(PlatformInfo);
+
+		if (PlatformInfo == nullptr)
+		{
+			return nullptr;
+		}
 
 		return FEditorStyle::GetBrush(PlatformInfo->GetIconStyleName(PlatformInfo::EPlatformIconSize::Normal));
 	}
