Docs: Removed file comments and added missing code documentation

diff --git a/Engine/Source/Developer/SessionFrontend/Private/Models/SessionBrowserOwnerFilter.h b/Engine/Source/Developer/SessionFrontend/Private/Models/SessionBrowserOwnerFilter.h
--- a/Engine/Source/Developer/SessionFrontend/Private/Models/SessionBrowserOwnerFilter.h
+++ b/Engine/Source/Developer/SessionFrontend/Private/Models/SessionBrowserOwnerFilter.h
@@ -1,9 +1,5 @@
 // Copyright 1998-2014 Epic Games, Inc. All Rights Reserved.
 
-/*=============================================================================
-	SessionBrowserOwnerFilter.h: Declares the FSessionBrowserOwnerFilter class.
-=============================================================================*/
-
 #pragma once
 
 
diff --git a/Engine/Source/Developer/SessionFrontend/Private/Models/SessionConsoleCategoryFilter.h b/Engine/Source/Developer/SessionFrontend/Private/Models/SessionConsoleCategoryFilter.h
--- a/Engine/Source/Developer/SessionFrontend/Private/Models/SessionConsoleCategoryFilter.h
+++ b/Engine/Source/Developer/SessionFrontend/Private/Models/SessionConsoleCategoryFilter.h
@@ -1,9 +1,5 @@
 // Copyright 1998-2014 Epic Games, Inc. All Rights Reserved.
 
-/*=============================================================================
-	SessionConsoleCategoryFilter.h: Declares the FSessionConsoleCategoryFilter class.
-=============================================================================*/
-
 #pragma once
 
 
diff --git a/Engine/Source/Developer/SessionFrontend/Private/Models/SessionConsoleCommands.h b/Engine/Source/Developer/SessionFrontend/Private/Models/SessionConsoleCommands.h
--- a/Engine/Source/Developer/SessionFrontend/Private/Models/SessionConsoleCommands.h
+++ b/Engine/Source/Developer/SessionFrontend/Private/Models/SessionConsoleCommands.h
@@ -1,9 +1,5 @@
 // Copyright 1998-2014 Epic Games, Inc. All Rights Reserved.
 
-/*=============================================================================
-	SessionConsoleCommands.h: Declares the FSessionConsoleCommands class.
-=============================================================================*/
-
 #pragma once
 
 
@@ -28,17 +24,15 @@ public:
 
 public:
 
-	// Begin TCommands interface
+	// TCommands interface
 
-	virtual void RegisterCommands( ) OVERRIDE
+	virtual void RegisterCommands( ) override
 	{
 		UI_COMMAND(Clear, "Clear Log", "Clear the log window", EUserInterfaceActionType::Button, FInputGesture());
 		UI_COMMAND(SessionCopy, "Copy", "Copy the selected log messages to the clipboard", EUserInterfaceActionType::Button, FInputGesture());
 		UI_COMMAND(SessionSave, "Save Log...", "Save the entire log to a file", EUserInterfaceActionType::ToggleButton, FInputGesture());
 	}
 
-	// End TCommands interface
-
 public:
 
 	TSharedPtr<FUICommandInfo> Clear;
diff --git a/Engine/Source/Developer/SessionFrontend/Private/Models/SessionConsoleLogMessage.h b/Engine/Source/Developer/SessionFrontend/Private/Models/SessionConsoleLogMessage.h
--- a/Engine/Source/Developer/SessionFrontend/Private/Models/SessionConsoleLogMessage.h
+++ b/Engine/Source/Developer/SessionFrontend/Private/Models/SessionConsoleLogMessage.h
@@ -1,9 +1,5 @@
 // Copyright 1998-2014 Epic Games, Inc. All Rights Reserved.
 
-/*=============================================================================
-	SessionConsoleLogMessage.h: Declares the FSessionConsoleLogMessage structure.
-=============================================================================*/
-
 #pragma once
 
 
diff --git a/Engine/Source/Developer/SessionFrontend/Private/Models/SessionConsoleVerbosityFilter.h b/Engine/Source/Developer/SessionFrontend/Private/Models/SessionConsoleVerbosityFilter.h
--- a/Engine/Source/Developer/SessionFrontend/Private/Models/SessionConsoleVerbosityFilter.h
+++ b/Engine/Source/Developer/SessionFrontend/Private/Models/SessionConsoleVerbosityFilter.h
@@ -1,9 +1,5 @@
 // Copyright 1998-2014 Epic Games, Inc. All Rights Reserved.
 
-/*=============================================================================
-	SessionConsoleVerbosityFilter.h: Declares the FSessionConsoleVerbosityFilter class.
-=============================================================================*/
-
 #pragma once
 
 
diff --git a/Engine/Source/Developer/SessionFrontend/Private/SessionFrontendModule.cpp b/Engine/Source/Developer/SessionFrontend/Private/SessionFrontendModule.cpp
--- a/Engine/Source/Developer/SessionFrontend/Private/SessionFrontendModule.cpp
+++ b/Engine/Source/Developer/SessionFrontend/Private/SessionFrontendModule.cpp
@@ -1,9 +1,5 @@
 // Copyright 1998-2014 Epic Games, Inc. All Rights Reserved.
 
-/*=============================================================================
-	SessionFrontendModule.cpp: Implements the FSessionFrontendModule class.
-=============================================================================*/
-
 #include "SessionFrontendPrivatePCH.h"
 
 
@@ -18,25 +14,23 @@ class FSessionFrontendModule
 {
 public:
 
-	// Begin ISessionFrontendModule interface
+	// ISessionFrontendModule interface
 	
-	virtual TSharedRef<class SWidget> CreateSessionBrowser( const ISessionManagerRef& SessionManager ) OVERRIDE
+	virtual TSharedRef<class SWidget> CreateSessionBrowser( const ISessionManagerRef& SessionManager ) override
 	{
 		return SNew(SSessionBrowser, SessionManager);
 	}
 	
-	virtual TSharedRef<class SWidget> CreateSessionConsole( const ISessionManagerRef& SessionManager ) OVERRIDE
+	virtual TSharedRef<class SWidget> CreateSessionConsole( const ISessionManagerRef& SessionManager ) override
 	{
 		return SNew(SSessionConsole, SessionManager);
 	}
 	
-	// End ISessionFrontendModule interface
-	
 public:
 
-	// Begin IModuleInterface interface
+	// IModuleInterface interface
 	
-	virtual void StartupModule( ) OVERRIDE
+	virtual void StartupModule( ) override
 	{
 		FGlobalTabmanager::Get()->RegisterTabSpawner(SessionFrontendTabName, FOnSpawnTab::CreateRaw(this, &FSessionFrontendModule::SpawnSessionFrontendTab))
 			.SetDisplayName(NSLOCTEXT("FSessionFrontendModule", "FrontendTabTitle", "Session Frontend"))
@@ -44,12 +38,10 @@ public:
 			.SetIcon(FSlateIcon(FEditorStyle::GetStyleSetName(), "SessionFrontEnd.TabIcon"));
 	}
 
-	virtual void ShutdownModule( ) OVERRIDE
+	virtual void ShutdownModule( ) override
 	{
 		FGlobalTabmanager::Get()->UnregisterTabSpawner(SessionFrontendTabName);
 	}
-	
-	// End IModuleInterface interface
 
 private:
 
@@ -57,7 +49,6 @@ private:
 	 * Creates a new session front-end tab.
 	 *
 	 * @param SpawnTabArgs - The arguments for the tab to spawn.
-	 *
 	 * @return The spawned tab.
 	 */
 	TSharedRef<SDockTab> SpawnSessionFrontendTab( const FSpawnTabArgs& SpawnTabArgs )
diff --git a/Engine/Source/Developer/SessionFrontend/Private/SessionFrontendPrivatePCH.h b/Engine/Source/Developer/SessionFrontend/Private/SessionFrontendPrivatePCH.h
--- a/Engine/Source/Developer/SessionFrontend/Private/SessionFrontendPrivatePCH.h
+++ b/Engine/Source/Developer/SessionFrontend/Private/SessionFrontendPrivatePCH.h
@@ -1,9 +1,5 @@
 // Copyright 1998-2014 Epic Games, Inc. All Rights Reserved.
 
-/*=============================================================================
-	SessionFrontendPrivatePCH.h: Pre-compiled header file for the SessionFrontend module.
-=============================================================================*/
-
 #pragma once
 
 #include "SessionFrontend.h"
diff --git a/Engine/Source/Developer/SessionFrontend/Private/Widgets/Browser/SSessionBrowser.cpp b/Engine/Source/Developer/SessionFrontend/Private/Widgets/Browser/SSessionBrowser.cpp
--- a/Engine/Source/Developer/SessionFrontend/Private/Widgets/Browser/SSessionBrowser.cpp
+++ b/Engine/Source/Developer/SessionFrontend/Private/Widgets/Browser/SSessionBrowser.cpp
@@ -1,9 +1,5 @@
 // Copyright 1998-2014 Epic Games, Inc. All Rights Reserved.
 
-/*=============================================================================
-	SSessionBrowser.cpp: Implements the SSessionBrowser class.
-=============================================================================*/
-
 #include "SessionFrontendPrivatePCH.h"
 
 
diff --git a/Engine/Source/Developer/SessionFrontend/Private/Widgets/Browser/SSessionBrowser.h b/Engine/Source/Developer/SessionFrontend/Private/Widgets/Browser/SSessionBrowser.h
--- a/Engine/Source/Developer/SessionFrontend/Private/Widgets/Browser/SSessionBrowser.h
+++ b/Engine/Source/Developer/SessionFrontend/Private/Widgets/Browser/SSessionBrowser.h
@@ -1,9 +1,5 @@
 // Copyright 1998-2014 Epic Games, Inc. All Rights Reserved.
 
-/*=============================================================================
-	SSessionBrowser.h: Declares the SSessionBrowser class.
-=============================================================================*/
-
 #pragma once
 
 
diff --git a/Engine/Source/Developer/SessionFrontend/Private/Widgets/Browser/SSessionBrowserCommandBar.cpp b/Engine/Source/Developer/SessionFrontend/Private/Widgets/Browser/SSessionBrowserCommandBar.cpp
--- a/Engine/Source/Developer/SessionFrontend/Private/Widgets/Browser/SSessionBrowserCommandBar.cpp
+++ b/Engine/Source/Developer/SessionFrontend/Private/Widgets/Browser/SSessionBrowserCommandBar.cpp
@@ -1,9 +1,5 @@
 // Copyright 1998-2014 Epic Games, Inc. All Rights Reserved.
 
-/*=============================================================================
-	SSessionBrowserCommandBar.cpp: Implements the SSessionBrowserCommandBar class.
-=============================================================================*/
-
 #include "SessionFrontendPrivatePCH.h"
 
 
diff --git a/Engine/Source/Developer/SessionFrontend/Private/Widgets/Browser/SSessionBrowserCommandBar.h b/Engine/Source/Developer/SessionFrontend/Private/Widgets/Browser/SSessionBrowserCommandBar.h
--- a/Engine/Source/Developer/SessionFrontend/Private/Widgets/Browser/SSessionBrowserCommandBar.h
+++ b/Engine/Source/Developer/SessionFrontend/Private/Widgets/Browser/SSessionBrowserCommandBar.h
@@ -1,9 +1,5 @@
 // Copyright 1998-2014 Epic Games, Inc. All Rights Reserved.
 
-/*=============================================================================
-	SSessionBrowserCommandBar.h: Declares the SSessionBrowserCommandBar class.
-=============================================================================*/
-
 #pragma once
 
 
diff --git a/Engine/Source/Developer/SessionFrontend/Private/Widgets/Browser/SSessionBrowserInstanceList.h b/Engine/Source/Developer/SessionFrontend/Private/Widgets/Browser/SSessionBrowserInstanceList.h
--- a/Engine/Source/Developer/SessionFrontend/Private/Widgets/Browser/SSessionBrowserInstanceList.h
+++ b/Engine/Source/Developer/SessionFrontend/Private/Widgets/Browser/SSessionBrowserInstanceList.h
@@ -1,9 +1,5 @@
 // Copyright 1998-2014 Epic Games, Inc. All Rights Reserved.
 
-/*=============================================================================
-	SSessionInstanceList.h: Declares the SSessionInstanceList class.
-=============================================================================*/
-
 #pragma once
 
 
diff --git a/Engine/Source/Developer/SessionFrontend/Private/Widgets/Browser/SSessionBrowserInstanceListRow.h b/Engine/Source/Developer/SessionFrontend/Private/Widgets/Browser/SSessionBrowserInstanceListRow.h
--- a/Engine/Source/Developer/SessionFrontend/Private/Widgets/Browser/SSessionBrowserInstanceListRow.h
+++ b/Engine/Source/Developer/SessionFrontend/Private/Widgets/Browser/SSessionBrowserInstanceListRow.h
@@ -1,9 +1,5 @@
 // Copyright 1998-2014 Epic Games, Inc. All Rights Reserved.
 
-/*=============================================================================
-	SSessionBrowserInstanceListRow.h: Declares the SSessionBrowserInstanceListRow class.
-=============================================================================*/
-
 #pragma once
 
 
diff --git a/Engine/Source/Developer/SessionFrontend/Private/Widgets/Browser/SSessionBrowserSessionListRow.h b/Engine/Source/Developer/SessionFrontend/Private/Widgets/Browser/SSessionBrowserSessionListRow.h
--- a/Engine/Source/Developer/SessionFrontend/Private/Widgets/Browser/SSessionBrowserSessionListRow.h
+++ b/Engine/Source/Developer/SessionFrontend/Private/Widgets/Browser/SSessionBrowserSessionListRow.h
@@ -1,9 +1,5 @@
 // Copyright 1998-2014 Epic Games, Inc. All Rights Reserved.
 
-/*=============================================================================
-	SSessionBrowserSessionListRow.h: Declares the SSessionBrowserSessionListRow class.
-=============================================================================*/
-
 #pragma once
 
 
diff --git a/Engine/Source/Developer/SessionFrontend/Private/Widgets/Console/SSessionConsole.cpp b/Engine/Source/Developer/SessionFrontend/Private/Widgets/Console/SSessionConsole.cpp
--- a/Engine/Source/Developer/SessionFrontend/Private/Widgets/Console/SSessionConsole.cpp
+++ b/Engine/Source/Developer/SessionFrontend/Private/Widgets/Console/SSessionConsole.cpp
@@ -1,9 +1,5 @@
 // Copyright 1998-2014 Epic Games, Inc. All Rights Reserved.
 
-/*=============================================================================
-	SSessionConsolePanel.cpp: Implements the SSessionConsolePanel class.
-=============================================================================*/
-
 #include "SessionFrontendPrivatePCH.h"
 
 
diff --git a/Engine/Source/Developer/SessionFrontend/Private/Widgets/Console/SSessionConsole.h b/Engine/Source/Developer/SessionFrontend/Private/Widgets/Console/SSessionConsole.h
--- a/Engine/Source/Developer/SessionFrontend/Private/Widgets/Console/SSessionConsole.h
+++ b/Engine/Source/Developer/SessionFrontend/Private/Widgets/Console/SSessionConsole.h
@@ -1,9 +1,5 @@
 // Copyright 1998-2014 Epic Games, Inc. All Rights Reserved.
 
-/*=============================================================================
-	SSessionConsolePanel.h: Declares the SSessionConsolePanel class.
-=============================================================================*/
-
 #pragma once
 
 
diff --git a/Engine/Source/Developer/SessionFrontend/Private/Widgets/Console/SSessionConsoleCommandBar.cpp b/Engine/Source/Developer/SessionFrontend/Private/Widgets/Console/SSessionConsoleCommandBar.cpp
--- a/Engine/Source/Developer/SessionFrontend/Private/Widgets/Console/SSessionConsoleCommandBar.cpp
+++ b/Engine/Source/Developer/SessionFrontend/Private/Widgets/Console/SSessionConsoleCommandBar.cpp
@@ -1,9 +1,5 @@
 // Copyright 1998-2014 Epic Games, Inc. All Rights Reserved.
 
-/*=============================================================================
-	SSessionConsoleCommandBar.cpp: Implements the SSessionConsoleCommandBar class.
-=============================================================================*/
-
 #include "SessionFrontendPrivatePCH.h"
 
 
diff --git a/Engine/Source/Developer/SessionFrontend/Private/Widgets/Console/SSessionConsoleCommandBar.h b/Engine/Source/Developer/SessionFrontend/Private/Widgets/Console/SSessionConsoleCommandBar.h
--- a/Engine/Source/Developer/SessionFrontend/Private/Widgets/Console/SSessionConsoleCommandBar.h
+++ b/Engine/Source/Developer/SessionFrontend/Private/Widgets/Console/SSessionConsoleCommandBar.h
@@ -1,9 +1,5 @@
 // Copyright 1998-2014 Epic Games, Inc. All Rights Reserved.
 
-/*=============================================================================
-	SSessionConsoleCommandBar.h: Declares the SSessionConsoleCommandBar class.
-=============================================================================*/
-
 #pragma once
 
 
@@ -25,14 +21,10 @@ public:
 
 	SLATE_BEGIN_ARGS(SSessionConsoleCommandBar) { }
 
-		/**
-		 * Called when the filter settings have changed.
-		 */
+		/** Called when the filter settings have changed. */
 		SLATE_EVENT(FOnSessionConsoleCommandSubmitted, OnCommandSubmitted)
 
-		/**
-		 * Called when the promote to shortcut button is clicked.
-		 */
+		/** Called when the promote to shortcut button is clicked. */
 		SLATE_EVENT(FOnSessionConsoleCommandSubmitted, OnPromoteToShortcutClicked)
 
 	SLATE_END_ARGS()
diff --git a/Engine/Source/Developer/SessionFrontend/Private/Widgets/Console/SSessionConsoleFilterBar.cpp b/Engine/Source/Developer/SessionFrontend/Private/Widgets/Console/SSessionConsoleFilterBar.cpp
--- a/Engine/Source/Developer/SessionFrontend/Private/Widgets/Console/SSessionConsoleFilterBar.cpp
+++ b/Engine/Source/Developer/SessionFrontend/Private/Widgets/Console/SSessionConsoleFilterBar.cpp
@@ -1,9 +1,5 @@
 // Copyright 1998-2014 Epic Games, Inc. All Rights Reserved.
 
-/*=============================================================================
-	SSessionConsoleFilterBar.cpp: Implements the SSessionConsoleFilterBar class.
-=============================================================================*/
-
 #include "SessionFrontendPrivatePCH.h"
 
 
diff --git a/Engine/Source/Developer/SessionFrontend/Private/Widgets/Console/SSessionConsoleFilterBar.h b/Engine/Source/Developer/SessionFrontend/Private/Widgets/Console/SSessionConsoleFilterBar.h
--- a/Engine/Source/Developer/SessionFrontend/Private/Widgets/Console/SSessionConsoleFilterBar.h
+++ b/Engine/Source/Developer/SessionFrontend/Private/Widgets/Console/SSessionConsoleFilterBar.h
@@ -1,9 +1,5 @@
 // Copyright 1998-2014 Epic Games, Inc. All Rights Reserved.
 
-/*=============================================================================
-	SSessionConsoleFilterBar.h: Declares the SSessionConsoleFilterBar class.
-=============================================================================*/
-
 #pragma once
 
 
@@ -17,9 +13,7 @@ public:
 
 	SLATE_BEGIN_ARGS(SSessionConsoleFilterBar) { }
 
-		/**
-		 * Called when the filter settings have changed.
-		 */
+		/** Called when the filter settings have changed. */
 		SLATE_EVENT(FSimpleDelegate, OnFilterChanged)
 
 	SLATE_END_ARGS()
diff --git a/Engine/Source/Developer/SessionFrontend/Private/Widgets/Console/SSessionConsoleLogTableRow.h b/Engine/Source/Developer/SessionFrontend/Private/Widgets/Console/SSessionConsoleLogTableRow.h
--- a/Engine/Source/Developer/SessionFrontend/Private/Widgets/Console/SSessionConsoleLogTableRow.h
+++ b/Engine/Source/Developer/SessionFrontend/Private/Widgets/Console/SSessionConsoleLogTableRow.h
@@ -1,9 +1,5 @@
 // Copyright 1998-2014 Epic Games, Inc. All Rights Reserved.
 
-/*=============================================================================
-	SSessionConsoleLogTableRow.h: Declares the SSessionConsoleLogTableRow class.
-=============================================================================*/
-
 #pragma once
 
 
@@ -37,9 +33,9 @@ public:
 
 public:
 
-	// Begin SMultiColumnTableRow interface
+	// SMultiColumnTableRow interface
 
-	virtual TSharedRef<SWidget> GenerateWidgetForColumn( const FName& ColumnName ) OVERRIDE
+	virtual TSharedRef<SWidget> GenerateWidgetForColumn( const FName& ColumnName ) override
 	{
 		if (ColumnName == "Instance")
 		{
@@ -114,8 +110,6 @@ public:
 		return SNullWidget::NullWidget;
 	}
 
-	// End SMultiColumnTableRow interface
-
 private:
 
 	// Gets the border color for this row.
diff --git a/Engine/Source/Developer/SessionFrontend/Private/Widgets/Console/SSessionConsoleShortcutWindow.cpp b/Engine/Source/Developer/SessionFrontend/Private/Widgets/Console/SSessionConsoleShortcutWindow.cpp
--- a/Engine/Source/Developer/SessionFrontend/Private/Widgets/Console/SSessionConsoleShortcutWindow.cpp
+++ b/Engine/Source/Developer/SessionFrontend/Private/Widgets/Console/SSessionConsoleShortcutWindow.cpp
@@ -1,9 +1,5 @@
 // Copyright 1998-2014 Epic Games, Inc. All Rights Reserved.
 
-/*=============================================================================
-	SSessionConsoleShortcutWindow.cpp: Implements the SSessionConsoleShortcutWindow class.
-=============================================================================*/
-
 #include "SessionFrontendPrivatePCH.h"
 
 
diff --git a/Engine/Source/Developer/SessionFrontend/Private/Widgets/Console/SSessionConsoleShortcutWindow.h b/Engine/Source/Developer/SessionFrontend/Private/Widgets/Console/SSessionConsoleShortcutWindow.h
--- a/Engine/Source/Developer/SessionFrontend/Private/Widgets/Console/SSessionConsoleShortcutWindow.h
+++ b/Engine/Source/Developer/SessionFrontend/Private/Widgets/Console/SSessionConsoleShortcutWindow.h
@@ -1,9 +1,5 @@
 // Copyright 1998-2014 Epic Games, Inc. All Rights Reserved.
 
-/*=============================================================================
-	SSessionConsoleShortcutWindow.h: Declares the SSessionConsoleShortcutWindow class.
-=============================================================================*/
-
 #pragma once
 
 
@@ -27,9 +23,7 @@ public:
 
 	SLATE_BEGIN_ARGS(SSessionConsoleShortcutWindow) { }
 
-		/**
-		 * Called when the filter settings have changed.
-		 */
+		/** Called when the filter settings have changed. */
 		SLATE_EVENT(FOnSessionConsoleCommandSubmitted, OnCommandSubmitted)
 
 	SLATE_END_ARGS()
diff --git a/Engine/Source/Developer/SessionFrontend/Private/Widgets/Console/SSessionConsoleToolbar.cpp b/Engine/Source/Developer/SessionFrontend/Private/Widgets/Console/SSessionConsoleToolbar.cpp
--- a/Engine/Source/Developer/SessionFrontend/Private/Widgets/Console/SSessionConsoleToolbar.cpp
+++ b/Engine/Source/Developer/SessionFrontend/Private/Widgets/Console/SSessionConsoleToolbar.cpp
@@ -1,21 +1,34 @@
 // Copyright 1998-2014 Epic Games, Inc. All Rights Reserved.
 
-/*=============================================================================
-	SSessionConsoleToolbar.cpp: Implements the SSessionConsoleToolbar class.
-=============================================================================*/
-
 #include "SessionFrontendPrivatePCH.h"
 
 
-
-
-
 /* SSessionConsoleToolbar interface
  *****************************************************************************/
 
-
-
-
-
-
-
+BEGIN_SLATE_FUNCTION_BUILD_OPTIMIZATION
+void SSessionConsoleToolbar::Construct( const FArguments& InArgs, const TSharedRef<FUICommandList>& CommandList )
+{
+	FSessionConsoleCommands::Register();
+
+	// create the toolbar
+	FToolBarBuilder Toolbar(CommandList, FMultiBoxCustomization::None);
+	{
+		Toolbar.AddToolBarButton(FSessionConsoleCommands::Get().SessionCopy);
+		Toolbar.AddSeparator();
+
+		Toolbar.AddToolBarButton(FSessionConsoleCommands::Get().Clear);
+		Toolbar.AddToolBarButton(FSessionConsoleCommands::Get().SessionSave);
+	}
+
+	ChildSlot
+	[
+		SNew(SBorder)
+			.BorderImage(FEditorStyle::GetBrush("ToolPanel.GroupBorder"))
+			.Padding(0.0f)
+			[
+				Toolbar.MakeWidget()
+			]
+	];
+}
+END_SLATE_FUNCTION_BUILD_OPTIMIZATION
diff --git a/Engine/Source/Developer/SessionFrontend/Private/Widgets/Console/SSessionConsoleToolbar.h b/Engine/Source/Developer/SessionFrontend/Private/Widgets/Console/SSessionConsoleToolbar.h
--- a/Engine/Source/Developer/SessionFrontend/Private/Widgets/Console/SSessionConsoleToolbar.h
+++ b/Engine/Source/Developer/SessionFrontend/Private/Widgets/Console/SSessionConsoleToolbar.h
@@ -1,9 +1,5 @@
 // Copyright 1998-2014 Epic Games, Inc. All Rights Reserved.
 
-/*=============================================================================
-	SSessionConsoleToolbar.h: Declares the SSessionConsoleToolbar class.
-=============================================================================*/
-
 #pragma once
 
 
@@ -29,32 +25,7 @@ public:
 	 * @param InArgs The construction arguments.
 	 * @param CommandList The command list to use.
 	 */
-	BEGIN_SLATE_FUNCTION_BUILD_OPTIMIZATION
-	void Construct( const FArguments& InArgs, const TSharedRef<FUICommandList>& CommandList )
-	{
-		FSessionConsoleCommands::Register();
-
-		// create the toolbar
-		FToolBarBuilder Toolbar(CommandList, FMultiBoxCustomization::None);
-		{
-			Toolbar.AddToolBarButton(FSessionConsoleCommands::Get().SessionCopy);
-			Toolbar.AddSeparator();
-
-			Toolbar.AddToolBarButton(FSessionConsoleCommands::Get().Clear);
-			Toolbar.AddToolBarButton(FSessionConsoleCommands::Get().SessionSave);
-		}
-
-		ChildSlot
-		[
-			SNew(SBorder)
-				.BorderImage(FEditorStyle::GetBrush("ToolPanel.GroupBorder"))
-				.Padding(0.0f)
-				[
-					Toolbar.MakeWidget()
-				]
-		];
-	}
-	END_SLATE_FUNCTION_BUILD_OPTIMIZATION
+	void Construct( const FArguments& InArgs, const TSharedRef<FUICommandList>& CommandList );
 };
 
 
diff --git a/Engine/Source/Developer/SessionFrontend/Private/Widgets/SSessionFrontend.cpp b/Engine/Source/Developer/SessionFrontend/Private/Widgets/SSessionFrontend.cpp
--- a/Engine/Source/Developer/SessionFrontend/Private/Widgets/SSessionFrontend.cpp
+++ b/Engine/Source/Developer/SessionFrontend/Private/Widgets/SSessionFrontend.cpp
@@ -1,9 +1,5 @@
 // Copyright 1998-2014 Epic Games, Inc. All Rights Reserved.
 
-/*=============================================================================
-	SSessionFrontend.cpp: Implements the SSessionFrontend class.
-=============================================================================*/
-
 #include "SessionFrontendPrivatePCH.h"
 
 
diff --git a/Engine/Source/Developer/SessionFrontend/Private/Widgets/SSessionFrontend.h b/Engine/Source/Developer/SessionFrontend/Private/Widgets/SSessionFrontend.h
--- a/Engine/Source/Developer/SessionFrontend/Private/Widgets/SSessionFrontend.h
+++ b/Engine/Source/Developer/SessionFrontend/Private/Widgets/SSessionFrontend.h
@@ -1,11 +1,8 @@
 // Copyright 1998-2014 Epic Games, Inc. All Rights Reserved.
 
-/*=============================================================================
-	SSessionFrontend.h: Declares the SSessionFrontend class.
-=============================================================================*/
-
 #pragma once
 
+
 /**
  * Implements the launcher application
  */
diff --git a/Engine/Source/Developer/SessionFrontend/Public/Interfaces/ISessionFrontendModule.h b/Engine/Source/Developer/SessionFrontend/Public/Interfaces/ISessionFrontendModule.h
--- a/Engine/Source/Developer/SessionFrontend/Public/Interfaces/ISessionFrontendModule.h
+++ b/Engine/Source/Developer/SessionFrontend/Public/Interfaces/ISessionFrontendModule.h
@@ -1,9 +1,5 @@
 // Copyright 1998-2014 Epic Games, Inc. All Rights Reserved.
 
-/*=============================================================================
-	ISessionFrontendModule.h: Declares the ISessionFrontendModule interface.
-=============================================================================*/
-
 #pragma once
 
 
diff --git a/Engine/Source/Developer/SessionFrontend/Public/SessionFrontend.h b/Engine/Source/Developer/SessionFrontend/Public/SessionFrontend.h
--- a/Engine/Source/Developer/SessionFrontend/Public/SessionFrontend.h
+++ b/Engine/Source/Developer/SessionFrontend/Public/SessionFrontend.h
@@ -1,9 +1,5 @@
 // Copyright 1998-2014 Epic Games, Inc. All Rights Reserved.
 
-/*=============================================================================
-	SessionFrontend.h: SessionFrontend module public header file.
-=============================================================================*/
-
 #pragma once
 
 
diff --git a/Engine/Source/Developer/SessionLauncher/Private/Models/SessionLauncherCommands.h b/Engine/Source/Developer/SessionLauncher/Private/Models/SessionLauncherCommands.h
--- a/Engine/Source/Developer/SessionLauncher/Private/Models/SessionLauncherCommands.h
+++ b/Engine/Source/Developer/SessionLauncher/Private/Models/SessionLauncherCommands.h
@@ -1,9 +1,5 @@
 // Copyright 1998-2014 Epic Games, Inc. All Rights Reserved.
 
-/*=============================================================================
-	SessionLauncherCommands.h: Declares the FSessionLauncherCommands class.
-=============================================================================*/
-
 #pragma once
 
 
@@ -28,10 +24,10 @@ public:
 	
 public:
 
-	// Begin TCommands interface
+	// TCommands interface
 
 	PRAGMA_DISABLE_OPTIMIZATION
-	virtual void RegisterCommands( ) OVERRIDE
+	virtual void RegisterCommands( ) override
 	{
 		UI_COMMAND(QuickLaunch, "Quick Launch", "Builds, cooks, and launches a build.", EUserInterfaceActionType::ToggleButton, FInputGesture(EModifierKey::Control, EKeys::L));
 		UI_COMMAND(CreateBuild, "Build", "Creates a build.", EUserInterfaceActionType::ToggleButton, FInputGesture(EModifierKey::Control, EKeys::B));
@@ -40,27 +36,17 @@ public:
 	}
 	PRAGMA_ENABLE_OPTIMIZATION
 
-	// End TCommands interface
-
 public:
 
-	/**
-	 * Toggles the data capture for all session instances. Global and custom command.
-	 */
-	TSharedPtr< FUICommandInfo > CreateBuild;
+	/** Toggles the data capture for all session instances. Global and custom command. */
+	TSharedPtr<FUICommandInfo> CreateBuild;
 
-	/**
-	 * Toggles the data preview for all session instances. Global and custom command.
-	 */
-	TSharedPtr< FUICommandInfo > DeployBuild;
+	/** Toggles the data preview for all session instances. Global and custom command. */
+	TSharedPtr<FUICommandInfo> DeployBuild;
 
-	/**
-	 * Toggles the data capture for all session instances. Global and custom command.
-	 */
-	TSharedPtr< FUICommandInfo > QuickLaunch;
+	/** Toggles the data capture for all session instances. Global and custom command. */
+	TSharedPtr<FUICommandInfo> QuickLaunch;
 
-	/**
-	 * Toggles the data capture for all session instances. Global and custom command.
-	 */
-	TSharedPtr< FUICommandInfo > AdvancedBuild;
+	/** Toggles the data capture for all session instances. Global and custom command. */
+	TSharedPtr<FUICommandInfo> AdvancedBuild;
 };
diff --git a/Engine/Source/Developer/SessionLauncher/Private/Models/SessionLauncherModel.h b/Engine/Source/Developer/SessionLauncher/Private/Models/SessionLauncherModel.h
--- a/Engine/Source/Developer/SessionLauncher/Private/Models/SessionLauncherModel.h
+++ b/Engine/Source/Developer/SessionLauncher/Private/Models/SessionLauncherModel.h
@@ -1,9 +1,5 @@
 // Copyright 1998-2014 Epic Games, Inc. All Rights Reserved.
 
-/*=============================================================================
-	SessionLauncherModel.h: Declares the FSessionLauncherModel class.
-=============================================================================*/
-
 #pragma once
 
 
@@ -26,14 +22,10 @@ namespace ESessionLauncherTasks
 }
 
 
-/**
- * Type definition for shared pointers to instances of FSessionLauncherModel.
- */
+/** Type definition for shared pointers to instances of FSessionLauncherModel. */
 typedef TSharedPtr<class FSessionLauncherModel> FSessionLauncherModelPtr;
 
-/**
- * Type definition for shared references to instances of FSessionLauncherModel.
- */
+/** Type definition for shared references to instances of FSessionLauncherModel. */
 typedef TSharedRef<class FSessionLauncherModel> FSessionLauncherModelRef;
 
 
diff --git a/Engine/Source/Developer/SessionLauncher/Private/SessionLauncherModule.cpp b/Engine/Source/Developer/SessionLauncher/Private/SessionLauncherModule.cpp
--- a/Engine/Source/Developer/SessionLauncher/Private/SessionLauncherModule.cpp
+++ b/Engine/Source/Developer/SessionLauncher/Private/SessionLauncherModule.cpp
@@ -1,9 +1,5 @@
 // Copyright 1998-2014 Epic Games, Inc. All Rights Reserved.
 
-/*=============================================================================
-	SessionLauncherModule.cpp: Implements the FSessionLauncherModule class.
-=============================================================================*/
-
 #include "SessionLauncherPrivatePCH.h"
 
 
@@ -18,9 +14,9 @@ class FSessionLauncherModule
 {
 public:
 
-	// Begin ISessionLauncherModule interface
+	// ISessionLauncherModule interface
 
-	virtual TSharedRef<class SWidget> CreateLauncherProgressPanel( const ILauncherWorkerRef& LauncherWorker ) OVERRIDE
+	virtual TSharedRef<class SWidget> CreateLauncherProgressPanel( const ILauncherWorkerRef& LauncherWorker ) override
 	{
 		TSharedRef<SSessionLauncherProgress> Panel = SNew(SSessionLauncherProgress);
 
@@ -29,13 +25,11 @@ public:
 		return Panel;
 	}
 
-	// End ISessionLauncherModule interface
-
 public:
 
-	// Begin IModuleInterface interface
+	// IModuleInterface interface
 	
-	virtual void StartupModule( ) OVERRIDE
+	virtual void StartupModule( ) override
 	{
 		FGlobalTabmanager::Get()->RegisterTabSpawner(SessionLauncherTabName, FOnSpawnTab::CreateRaw(this, &FSessionLauncherModule::SpawnSessionLauncherTab))
 			.SetDisplayName(NSLOCTEXT("FSessionLauncherModule", "LauncherTabTitle", "Game Launcher"))
@@ -43,20 +37,17 @@ public:
 			.SetIcon(FSlateIcon(FEditorStyle::GetStyleSetName(), "SessionLauncher.TabIcon"));
 	}
 
-	virtual void ShutdownModule( ) OVERRIDE
+	virtual void ShutdownModule( ) override
 	{
 		FGlobalTabmanager::Get()->UnregisterTabSpawner(SessionLauncherTabName);
 	}
-	
-	// End IModuleInterface interface
 
 private:
 
 	/**
 	 * Creates a new session launcher tab.
 	 *
-	 * @param SpawnTabArgs - The arguments for the tab to spawn.
-	 *
+	 * @param SpawnTabArgs The arguments for the tab to spawn.
 	 * @return The spawned tab.
 	 */
 	TSharedRef<SDockTab> SpawnSessionLauncherTab( const FSpawnTabArgs& SpawnTabArgs )
diff --git a/Engine/Source/Developer/SessionLauncher/Private/SessionLauncherPrivatePCH.h b/Engine/Source/Developer/SessionLauncher/Private/SessionLauncherPrivatePCH.h
--- a/Engine/Source/Developer/SessionLauncher/Private/SessionLauncherPrivatePCH.h
+++ b/Engine/Source/Developer/SessionLauncher/Private/SessionLauncherPrivatePCH.h
@@ -1,17 +1,11 @@
 // Copyright 1998-2014 Epic Games, Inc. All Rights Reserved.
 
-/*=============================================================================
-	SessionLauncherPrivatePCH.h: Pre-compiled header file for the SessionLauncher module.
-=============================================================================*/
+#pragma once
 
-#ifndef UNREAL_SESSIONLAUNCHER_PRIVATEPCH_H
-#define UNREAL_SESSIONLAUNCHER_PRIVATEPCH_H
+#include "SessionLauncher.h"
 
 
-#include "../Public/SessionLauncher.h"
-
-
-/* Dependencies
+/* Private dependencies
  *****************************************************************************/
 
 #include "DesktopPlatformModule.h"
@@ -97,6 +91,3 @@ static TSharedRef<SButton> MakeImageButton( const FSlateBrush* ButtonImage, cons
 #include "SSessionLauncherToolbar.h"
 #include "SSessionLauncherValidation.h"
 #include "SSessionLauncher.h"
-
-
-#endif //UNREAL_SESSIONLAUNCHER_PRIVATEPCH_H
diff --git a/Engine/Source/Developer/SessionLauncher/Private/Widgets/Build/SSessionLauncherBuildPage.cpp b/Engine/Source/Developer/SessionLauncher/Private/Widgets/Build/SSessionLauncherBuildPage.cpp
--- a/Engine/Source/Developer/SessionLauncher/Private/Widgets/Build/SSessionLauncherBuildPage.cpp
+++ b/Engine/Source/Developer/SessionLauncher/Private/Widgets/Build/SSessionLauncherBuildPage.cpp
@@ -92,6 +92,32 @@ void SSessionLauncherBuildPage::Construct( const FArguments& InArgs, const FSess
 }
 
 
+/* SSessionLauncherBuildPage implementation
+ *****************************************************************************/
+
+bool SSessionLauncherBuildPage::GenerateDSYMForProject( const FString& ProjectName, const FString& Configuration )
+{
+    // UAT executable
+    FString ExecutablePath = FPaths::ConvertRelativePathToFull(FPaths::EngineDir() + FString(TEXT("Build")) / TEXT("BatchFiles"));
+#if PLATFORM_MAC
+    FString Executable = TEXT("RunUAT.command");
+#else
+    FString Executable = TEXT("RunUAT.bat");
+#endif
+
+    // build UAT command line parameters
+    FString CommandLine;
+    CommandLine = FString::Printf(TEXT("GenerateDSYM -project=%s -config=%s"),
+        *ProjectName,
+        *Configuration);
+
+    // launch the builder and monitor its process
+    FProcHandle ProcessHandle = FPlatformProcess::CreateProc(*(ExecutablePath / Executable), *CommandLine, false, false, false, NULL, 0, *ExecutablePath, NULL);
+
+	return ProcessHandle.Close();
+}
+
+
 /* SSessionLauncherBuildPage callbacks
  *****************************************************************************/
 
@@ -177,25 +203,5 @@ bool SSessionLauncherBuildPage::HandleGenDSYMButtonEnabled() const
     return false;
 }
 
-bool SSessionLauncherBuildPage::GenerateDSYMForProject( const FString& ProjectName, const FString& Configuration )
-{
-    // UAT executable
-    FString ExecutablePath = FPaths::ConvertRelativePathToFull(FPaths::EngineDir() + FString(TEXT("Build")) / TEXT("BatchFiles"));
-#if PLATFORM_MAC
-    FString Executable = TEXT("RunUAT.command");
-#else
-    FString Executable = TEXT("RunUAT.bat");
-#endif
-
-    // build UAT command line parameters
-    FString CommandLine;
-    CommandLine = FString::Printf(TEXT("GenerateDSYM -project=%s -config=%s"),
-        *ProjectName,
-        *Configuration);
-
-    // launch the builder and monitor its process
-    FProcHandle ProcessHandle = FPlatformProcess::CreateProc(*(ExecutablePath / Executable), *CommandLine, false, false, false, NULL, 0, *ExecutablePath, NULL);
-    return ProcessHandle.Close();
-}
 
 #undef LOCTEXT_NAMESPACE
diff --git a/Engine/Source/Developer/SessionLauncher/Private/Widgets/Build/SSessionLauncherBuildPage.h b/Engine/Source/Developer/SessionLauncher/Private/Widgets/Build/SSessionLauncherBuildPage.h
--- a/Engine/Source/Developer/SessionLauncher/Private/Widgets/Build/SSessionLauncherBuildPage.h
+++ b/Engine/Source/Developer/SessionLauncher/Private/Widgets/Build/SSessionLauncherBuildPage.h
@@ -1,9 +1,5 @@
 // Copyright 1998-2014 Epic Games, Inc. All Rights Reserved.
 
-/*=============================================================================
-	SSessionLauncherBuildPage.h: Declares the SSessionLauncherBuildPage class.
-=============================================================================*/
-
 #pragma once
 
 
@@ -18,7 +14,6 @@ public:
 	SLATE_BEGIN_ARGS(SSessionLauncherBuildPage) { }
 	SLATE_END_ARGS()
 
-
 public:
 
 	/**
@@ -26,17 +21,26 @@ public:
 	 */
 	~SSessionLauncherBuildPage( );
 
-
 public:
 
 	/**
 	 * Constructs the widget.
 	 *
-	 * @param InArgs - The Slate argument list.
-	 * @param InModel - The data model.
+	 * @param InArgs The Slate argument list.
+	 * @param InModel The data model.
 	 */
 	void Construct(	const FArguments& InArgs, const FSessionLauncherModelRef& InModel );
 
+protected:
+
+	/**
+	 * Generates the DSYM for the project.
+	 *
+	 * @param ProjectName The name of the project.
+	 * @param Configuration The build configuration name.
+	 * @return true on success, false otherwise.
+	 */
+    bool GenerateDSYMForProject( const FString& ProjectName, const FString& Configuration );
 
 private:
 
@@ -54,8 +58,9 @@ private:
 
     // Callback for pressing the Advanced Setting - Generate DSYM button.
     FReply HandleGenDSYMClicked();
+
+	// Callback for getting the enabled state of the Generate DSYM button.
     bool HandleGenDSYMButtonEnabled() const;
-    bool GenerateDSYMForProject( const FString& ProjectName, const FString& Configuration );
 
 private:
 
diff --git a/Engine/Source/Developer/SessionLauncher/Private/Widgets/Cook/SSessionLauncherCookByTheBookSettings.cpp b/Engine/Source/Developer/SessionLauncher/Private/Widgets/Cook/SSessionLauncherCookByTheBookSettings.cpp
--- a/Engine/Source/Developer/SessionLauncher/Private/Widgets/Cook/SSessionLauncherCookByTheBookSettings.cpp
+++ b/Engine/Source/Developer/SessionLauncher/Private/Widgets/Cook/SSessionLauncherCookByTheBookSettings.cpp
@@ -1,9 +1,5 @@
 // Copyright 1998-2014 Epic Games, Inc. All Rights Reserved.
 
-/*=============================================================================
-	SSessionLauncherCookByTheBookSettings.cpp: Implements the SSessionLauncherCookByTheBookSettings class.
-=============================================================================*/
-
 #include "SessionLauncherPrivatePCH.h"
 
 
diff --git a/Engine/Source/Developer/SessionLauncher/Private/Widgets/Cook/SSessionLauncherCookByTheBookSettings.h b/Engine/Source/Developer/SessionLauncher/Private/Widgets/Cook/SSessionLauncherCookByTheBookSettings.h
--- a/Engine/Source/Developer/SessionLauncher/Private/Widgets/Cook/SSessionLauncherCookByTheBookSettings.h
+++ b/Engine/Source/Developer/SessionLauncher/Private/Widgets/Cook/SSessionLauncherCookByTheBookSettings.h
@@ -49,8 +49,8 @@ public:
 	/**
 	 * Constructs the widget.
 	 *
-	 * @param InArgs - The Slate argument list.
-	 * @param InModel - The data model.
+	 * @param InArgs The Slate argument list.
+	 * @param InModel The data model.
 	 */
 	void Construct(	const FArguments& InArgs, const FSessionLauncherModelRef& InModel, bool InShowSimple = false );
 
diff --git a/Engine/Source/Developer/SessionLauncher/Private/Widgets/Cook/SSessionLauncherCookOnTheFlySettings.cpp b/Engine/Source/Developer/SessionLauncher/Private/Widgets/Cook/SSessionLauncherCookOnTheFlySettings.cpp
--- a/Engine/Source/Developer/SessionLauncher/Private/Widgets/Cook/SSessionLauncherCookOnTheFlySettings.cpp
+++ b/Engine/Source/Developer/SessionLauncher/Private/Widgets/Cook/SSessionLauncherCookOnTheFlySettings.cpp
@@ -1,9 +1,5 @@
 // Copyright 1998-2014 Epic Games, Inc. All Rights Reserved.
 
-/*=============================================================================
-	SSessionLauncherCookOnTheFlySettings.cpp: Implements the SSessionLauncherCookOnTheFlySettings class.
-=============================================================================*/
-
 #include "SessionLauncherPrivatePCH.h"
 
 
diff --git a/Engine/Source/Developer/SessionLauncher/Private/Widgets/Cook/SSessionLauncherCookOnTheFlySettings.h b/Engine/Source/Developer/SessionLauncher/Private/Widgets/Cook/SSessionLauncherCookOnTheFlySettings.h
--- a/Engine/Source/Developer/SessionLauncher/Private/Widgets/Cook/SSessionLauncherCookOnTheFlySettings.h
+++ b/Engine/Source/Developer/SessionLauncher/Private/Widgets/Cook/SSessionLauncherCookOnTheFlySettings.h
@@ -1,9 +1,5 @@
 // Copyright 1998-2014 Epic Games, Inc. All Rights Reserved.
 
-/*=============================================================================
-	SSessionLauncherCookOnTheFlySettings.h: Declares the SSessionLauncherCookOnTheFlySettings class.
-=============================================================================*/
-
 #pragma once
 
 
@@ -18,7 +14,6 @@ public:
 	SLATE_BEGIN_ARGS(SSessionLauncherCookOnTheFlySettings) { }
 	SLATE_END_ARGS()
 
-
 public:
 
 	/**
@@ -26,18 +21,16 @@ public:
 	 */
 	~SSessionLauncherCookOnTheFlySettings( );
 
-
 public:
 
 	/**
 	 * Constructs the widget.
 	 *
-	 * @param InArgs - The Slate argument list.
-	 * @param InModel - The data model.
+	 * @param InArgs The Slate argument list.
+	 * @param InModel The data model.
 	 */
 	void Construct(	const FArguments& InArgs, const FSessionLauncherModelRef& InModel );
 
-
 private:
 
 	// Callback for check state changes of the 'Incremental' check box.
@@ -52,7 +45,6 @@ private:
 	// Callback for determining the visibility of a validation error icon.
 	EVisibility HandleValidationErrorIconVisibility( ELauncherProfileValidationErrors::Type Error ) const;
 
-
 private:
 
 	// Holds a pointer to the data model.
diff --git a/Engine/Source/Developer/SessionLauncher/Private/Widgets/Cook/SSessionLauncherCookPage.cpp b/Engine/Source/Developer/SessionLauncher/Private/Widgets/Cook/SSessionLauncherCookPage.cpp
--- a/Engine/Source/Developer/SessionLauncher/Private/Widgets/Cook/SSessionLauncherCookPage.cpp
+++ b/Engine/Source/Developer/SessionLauncher/Private/Widgets/Cook/SSessionLauncherCookPage.cpp
@@ -1,9 +1,5 @@
 // Copyright 1998-2014 Epic Games, Inc. All Rights Reserved.
 
-/*=============================================================================
-	SSessionLauncherCookPage.cpp: Implements the SSessionLauncherCookPage class.
-=============================================================================*/
-
 #include "SessionLauncherPrivatePCH.h"
 
 
diff --git a/Engine/Source/Developer/SessionLauncher/Private/Widgets/Cook/SSessionLauncherCookPage.h b/Engine/Source/Developer/SessionLauncher/Private/Widgets/Cook/SSessionLauncherCookPage.h
--- a/Engine/Source/Developer/SessionLauncher/Private/Widgets/Cook/SSessionLauncherCookPage.h
+++ b/Engine/Source/Developer/SessionLauncher/Private/Widgets/Cook/SSessionLauncherCookPage.h
@@ -1,9 +1,5 @@
 // Copyright 1998-2014 Epic Games, Inc. All Rights Reserved.
 
-/*=============================================================================
-	SSessionLauncherCookPage.h: Declares the SSessionLauncherCookPage class.
-=============================================================================*/
-
 #pragma once
 
 
@@ -18,7 +14,6 @@ public:
 	SLATE_BEGIN_ARGS(SSessionLauncherCookPage) { }
 	SLATE_END_ARGS()
 
-
 public:
 
 	/**
@@ -26,18 +21,16 @@ public:
 	 */
 	~SSessionLauncherCookPage( );
 
-
 public:
 
 	/**
 	 * Constructs the widget.
 	 *
-	 * @param InArgs - The Slate argument list.
-	 * @param InModel - The data model.
+	 * @param InArgs The Slate argument list.
+	 * @param InModel The data model.
 	 */
 	void Construct(	const FArguments& InArgs, const FSessionLauncherModelRef& InModel );
 
-
 private:
 
 	// Callback for getting the visibility of the cook-by-the-book settings area.
@@ -58,7 +51,6 @@ private:
 	// Callback for changing the selected profile in the profile manager.
 	void HandleProfileManagerProfileSelected( const ILauncherProfilePtr& SelectedProfile, const ILauncherProfilePtr& PreviousProfile );
 
-
 private:
 
 	// Holds the cooker options text box.
diff --git a/Engine/Source/Developer/SessionLauncher/Private/Widgets/Cook/SSessionLauncherCookedPlatforms.h b/Engine/Source/Developer/SessionLauncher/Private/Widgets/Cook/SSessionLauncherCookedPlatforms.h
--- a/Engine/Source/Developer/SessionLauncher/Private/Widgets/Cook/SSessionLauncherCookedPlatforms.h
+++ b/Engine/Source/Developer/SessionLauncher/Private/Widgets/Cook/SSessionLauncherCookedPlatforms.h
@@ -1,9 +1,5 @@
 // Copyright 1998-2014 Epic Games, Inc. All Rights Reserved.
 
-/*=============================================================================
-	SSessionLauncherCookedPlatforms.h: Declares the SSessionLauncherCookedPlatforms class.
-=============================================================================*/
-
 #pragma once
 
 
@@ -21,14 +17,13 @@ public:
 	SLATE_BEGIN_ARGS(SSessionLauncherCookedPlatforms) { }
 	SLATE_END_ARGS()
 
-
 public:
 
 	/**
 	 * Constructs the widget.
 	 *
-	 * @param InArgs - The Slate argument list.
-	 * @param InModel - The data model.
+	 * @param InArgs The Slate argument list.
+	 * @param InModel The data model.
 	 */
 	void Construct(	const FArguments& InArgs, const FSessionLauncherModelRef& InModel )
 	{
@@ -92,7 +87,6 @@ public:
 		];
 	}
 
-
 protected:
 
 	/**
@@ -137,7 +131,6 @@ protected:
 			.Text(LOCTEXT("NoPlatformsFoundErrorText", "Error: No platforms found."));
 	}
 
-
 private:
 
 	// Callback for clicking the 'Select All Platforms' button.
@@ -229,7 +222,6 @@ private:
 		return ESlateCheckBoxState::Unchecked;
 	}
 
-
 private:
 
 	// Holds a pointer to the data model.
diff --git a/Engine/Source/Developer/SessionLauncher/Private/Widgets/Cook/SSessionLauncherCultureListRow.h b/Engine/Source/Developer/SessionLauncher/Private/Widgets/Cook/SSessionLauncherCultureListRow.h
--- a/Engine/Source/Developer/SessionLauncher/Private/Widgets/Cook/SSessionLauncherCultureListRow.h
+++ b/Engine/Source/Developer/SessionLauncher/Private/Widgets/Cook/SSessionLauncherCultureListRow.h
@@ -1,9 +1,5 @@
 // Copyright 1998-2014 Epic Games, Inc. All Rights Reserved.
 
-/*=============================================================================
-	SSessionLauncherCultureListRow.h: Declares the SSessionLauncherCultureListRow class.
-=============================================================================*/
-
 #pragma once
 
 
@@ -24,14 +20,13 @@ public:
 		SLATE_ARGUMENT(TSharedPtr<FString>, CultureName)
 	SLATE_END_ARGS()
 
-
 public:
 
 	/**
 	 * Constructs the widget.
 	 *
-	 * @param InArgs - The construction arguments.
-	 * @param InProfileManager - The profile manager to use.
+	 * @param InArgs The construction arguments.
+	 * @param InProfileManager The profile manager to use.
 	 */
 	void Construct( const FArguments& InArgs, const FSessionLauncherModelRef& InModel )
 	{
@@ -42,13 +37,12 @@ public:
 		SMultiColumnTableRow<TSharedPtr<FString> >::Construct(FSuperRowType::FArguments(), InArgs._OwnerTableView.ToSharedRef());
 	}
 
-
 public:
 
 	/**
 	 * Generates the widget for the specified column.
 	 *
-	 * @param ColumnName - The name of the column to generate the widget for.
+	 * @param ColumnName The name of the column to generate the widget for.
 	 *
 	 * @return The widget.
 	 */
@@ -69,7 +63,6 @@ public:
 		return SNullWidget::NullWidget;
 	}
 
-
 private:
 
 	// Callback for changing the checked state of the check box.
@@ -103,7 +96,6 @@ private:
 		return ESlateCheckBoxState::Unchecked;
 	}
 
-
 private:
 
 	// Holds the highlight string for the log message.
diff --git a/Engine/Source/Developer/SessionLauncher/Private/Widgets/Cook/SSessionLauncherMapListRow.h b/Engine/Source/Developer/SessionLauncher/Private/Widgets/Cook/SSessionLauncherMapListRow.h
--- a/Engine/Source/Developer/SessionLauncher/Private/Widgets/Cook/SSessionLauncherMapListRow.h
+++ b/Engine/Source/Developer/SessionLauncher/Private/Widgets/Cook/SSessionLauncherMapListRow.h
@@ -1,9 +1,5 @@
 // Copyright 1998-2014 Epic Games, Inc. All Rights Reserved.
 
-/*=============================================================================
-	SSessionLauncherMapListRow.h: Declares the SSessionLauncherMapListRow class.
-=============================================================================*/
-
 #pragma once
 
 
@@ -24,14 +20,13 @@ public:
 		SLATE_ARGUMENT(TSharedPtr<FString>, MapName)
 	SLATE_END_ARGS()
 
-
 public:
 
 	/**
 	 * Constructs the widget.
 	 *
-	 * @param InArgs - The construction arguments.
-	 * @param InProfileManager - The profile manager to use.
+	 * @param InArgs The construction arguments.
+	 * @param InProfileManager The profile manager to use.
 	 */
 	void Construct( const FArguments& InArgs, const FSessionLauncherModelRef& InModel )
 	{
@@ -42,14 +37,12 @@ public:
 		SMultiColumnTableRow<TSharedPtr<FString> >::Construct(FSuperRowType::FArguments(), InArgs._OwnerTableView.ToSharedRef());
 	}
 
-
 public:
 
 	/**
 	 * Generates the widget for the specified column.
 	 *
-	 * @param ColumnName - The name of the column to generate the widget for.
-	 *
+	 * @param ColumnName The name of the column to generate the widget for.
 	 * @return The widget.
 	 */
 	virtual TSharedRef<SWidget> GenerateWidgetForColumn( const FName& ColumnName ) OVERRIDE
@@ -69,7 +62,6 @@ public:
 		return SNullWidget::NullWidget;
 	}
 
-
 private:
 
 	// Callback for changing the checked state of the check box.
@@ -103,7 +95,6 @@ private:
 		return ESlateCheckBoxState::Unchecked;
 	}
 
-
 private:
 
 	// Holds the highlight string for the log message.
diff --git a/Engine/Source/Developer/SessionLauncher/Private/Widgets/Cook/SSessionLauncherSimpleCookPage.cpp b/Engine/Source/Developer/SessionLauncher/Private/Widgets/Cook/SSessionLauncherSimpleCookPage.cpp
--- a/Engine/Source/Developer/SessionLauncher/Private/Widgets/Cook/SSessionLauncherSimpleCookPage.cpp
+++ b/Engine/Source/Developer/SessionLauncher/Private/Widgets/Cook/SSessionLauncherSimpleCookPage.cpp
@@ -1,9 +1,5 @@
 // Copyright 1998-2014 Epic Games, Inc. All Rights Reserved.
 
-/*=============================================================================
-	SSessionLauncherSimpleCookPage.cpp: Implements the SSessionLauncherSimpleCookPage class.
-=============================================================================*/
-
 #include "SessionLauncherPrivatePCH.h"
 
 
diff --git a/Engine/Source/Developer/SessionLauncher/Private/Widgets/Cook/SSessionLauncherSimpleCookPage.h b/Engine/Source/Developer/SessionLauncher/Private/Widgets/Cook/SSessionLauncherSimpleCookPage.h
--- a/Engine/Source/Developer/SessionLauncher/Private/Widgets/Cook/SSessionLauncherSimpleCookPage.h
+++ b/Engine/Source/Developer/SessionLauncher/Private/Widgets/Cook/SSessionLauncherSimpleCookPage.h
@@ -1,9 +1,5 @@
 // Copyright 1998-2014 Epic Games, Inc. All Rights Reserved.
 
-/*=============================================================================
-	SSessionLauncherSimpleCookPage.h: Declares the SSessionLauncherSimpleCookPage class.
-=============================================================================*/
-
 #pragma once
 
 
@@ -18,7 +14,6 @@ public:
 	SLATE_BEGIN_ARGS(SSessionLauncherSimpleCookPage) { }
 	SLATE_END_ARGS()
 
-
 public:
 
 	/**
@@ -26,18 +21,16 @@ public:
 	 */
 	~SSessionLauncherSimpleCookPage( );
 
-
 public:
 
 	/**
 	 * Constructs the widget.
 	 *
-	 * @param InArgs - The Slate argument list.
-	 * @param InModel - The data model.
+	 * @param InArgs The Slate argument list.
+	 * @param InModel The data model.
 	 */
 	void Construct(	const FArguments& InArgs, const FSessionLauncherModelRef& InModel );
 
-
 private:
 
 	// Callback for getting the visibility of the cook-by-the-book settings area.
@@ -58,7 +51,6 @@ private:
 	// Callback for changing the selected profile in the profile manager.
 	void HandleProfileManagerProfileSelected( const ILauncherProfilePtr& SelectedProfile, const ILauncherProfilePtr& PreviousProfile );
 
-
 private:
 
 	// Holds the cooker options text box.
diff --git a/Engine/Source/Developer/SessionLauncher/Private/Widgets/Deploy/SSessionLauncherDeployFileServerSettings.cpp b/Engine/Source/Developer/SessionLauncher/Private/Widgets/Deploy/SSessionLauncherDeployFileServerSettings.cpp
--- a/Engine/Source/Developer/SessionLauncher/Private/Widgets/Deploy/SSessionLauncherDeployFileServerSettings.cpp
+++ b/Engine/Source/Developer/SessionLauncher/Private/Widgets/Deploy/SSessionLauncherDeployFileServerSettings.cpp
@@ -1,9 +1,5 @@
 // Copyright 1998-2014 Epic Games, Inc. All Rights Reserved.
 
-/*=============================================================================
-	SSessionLauncherDeployFileServerSettings.cpp: Implements the SSessionLauncherDeployFileServerSettings class.
-=============================================================================*/
-
 #include "SessionLauncherPrivatePCH.h"
 
 
diff --git a/Engine/Source/Developer/SessionLauncher/Private/Widgets/Deploy/SSessionLauncherDeployFileServerSettings.h b/Engine/Source/Developer/SessionLauncher/Private/Widgets/Deploy/SSessionLauncherDeployFileServerSettings.h
--- a/Engine/Source/Developer/SessionLauncher/Private/Widgets/Deploy/SSessionLauncherDeployFileServerSettings.h
+++ b/Engine/Source/Developer/SessionLauncher/Private/Widgets/Deploy/SSessionLauncherDeployFileServerSettings.h
@@ -1,9 +1,5 @@
 // Copyright 1998-2014 Epic Games, Inc. All Rights Reserved.
 
-/*=============================================================================
-	SSessionLauncherDeployFileServerSettings.h: Declares the SSessionLauncherDeployFileServerSettings class.
-=============================================================================*/
-
 #pragma once
 
 
@@ -18,14 +14,13 @@ public:
 	SLATE_BEGIN_ARGS(SSessionLauncherDeployFileServerSettings) { }
 	SLATE_END_ARGS()
 
-
 public:
 
 	/**
 	 * Constructs the widget.
 	 *
-	 * @param InArgs - The Slate argument list.
-	 * @param InModel - The data model.
+	 * @param InArgs The Slate argument list.
+	 * @param InModel The data model.
 	 */
 	void Construct(	const FArguments& InArgs, const FSessionLauncherModelRef& InModel );
 
@@ -44,7 +39,6 @@ private:
 	// Callback for determining the checked state of the 'Streaming Server' check box.
 	ESlateCheckBoxState::Type HandleStreamingServerCheckBoxIsChecked( ) const;
 
-
 private:
 
 	// Holds a pointer to the data model.
diff --git a/Engine/Source/Developer/SessionLauncher/Private/Widgets/Deploy/SSessionLauncherDeployPage.cpp b/Engine/Source/Developer/SessionLauncher/Private/Widgets/Deploy/SSessionLauncherDeployPage.cpp
--- a/Engine/Source/Developer/SessionLauncher/Private/Widgets/Deploy/SSessionLauncherDeployPage.cpp
+++ b/Engine/Source/Developer/SessionLauncher/Private/Widgets/Deploy/SSessionLauncherDeployPage.cpp
@@ -1,9 +1,5 @@
 // Copyright 1998-2014 Epic Games, Inc. All Rights Reserved.
 
-/*=============================================================================
-	SSessionLauncherDevicesPage.cpp: Implements the SSessionLauncherDevicesPage class.
-=============================================================================*/
-
 #include "SessionLauncherPrivatePCH.h"
 
 
diff --git a/Engine/Source/Developer/SessionLauncher/Private/Widgets/Deploy/SSessionLauncherDeployPage.h b/Engine/Source/Developer/SessionLauncher/Private/Widgets/Deploy/SSessionLauncherDeployPage.h
--- a/Engine/Source/Developer/SessionLauncher/Private/Widgets/Deploy/SSessionLauncherDeployPage.h
+++ b/Engine/Source/Developer/SessionLauncher/Private/Widgets/Deploy/SSessionLauncherDeployPage.h
@@ -1,9 +1,5 @@
 // Copyright 1998-2014 Epic Games, Inc. All Rights Reserved.
 
-/*=============================================================================
-	SSessionLauncherDevicesPage.h: Declares the SSessionLauncherDevicesPage class.
-=============================================================================*/
-
 #pragma once
 
 
@@ -18,7 +14,6 @@ public:
 	SLATE_BEGIN_ARGS(SSessionLauncherDeployPage) { }
 	SLATE_END_ARGS()
 
-
 public:
 
 	/**
@@ -26,7 +21,6 @@ public:
 	 */
 	~SSessionLauncherDeployPage( );
 
-
 public:
 
 	/**
@@ -37,7 +31,6 @@ public:
 	 */
 	void Construct(	const FArguments& InArgs, const FSessionLauncherModelRef& InModel, bool IsFromRepository = false );
 
-
 private:
 
 	// Callback for getting the content text of the 'Deployment Mode' combo button.
@@ -61,7 +54,6 @@ private:
 	// Callback for determining the visibility of a validation error icon.
 	EVisibility HandleValidationErrorIconVisibility( ELauncherProfileValidationErrors::Type Error ) const;
 
-
 private:
 
 	// Holds a pointer to the data model.
diff --git a/Engine/Source/Developer/SessionLauncher/Private/Widgets/Deploy/SSessionLauncherDeployRepositorySettings.cpp b/Engine/Source/Developer/SessionLauncher/Private/Widgets/Deploy/SSessionLauncherDeployRepositorySettings.cpp
--- a/Engine/Source/Developer/SessionLauncher/Private/Widgets/Deploy/SSessionLauncherDeployRepositorySettings.cpp
+++ b/Engine/Source/Developer/SessionLauncher/Private/Widgets/Deploy/SSessionLauncherDeployRepositorySettings.cpp
@@ -1,9 +1,5 @@
 // Copyright 1998-2014 Epic Games, Inc. All Rights Reserved.
 
-/*=============================================================================
-	SSessionLauncherDeployRepositorySettings.cpp: Implements the SSessionLauncherDeployRepositorySettings class.
-=============================================================================*/
-
 #include "SessionLauncherPrivatePCH.h"
 
 
@@ -125,6 +121,7 @@ FReply SSessionLauncherDeployRepositorySettings::HandleBrowseButtonClicked( )
 	return FReply::Handled();
 }
 
+
 void SSessionLauncherDeployRepositorySettings::OnTextChanged(const FText& InText)
 {
 	ILauncherProfilePtr SelectedProfile = Model->GetSelectedProfile();
@@ -135,6 +132,7 @@ void SSessionLauncherDeployRepositorySettings::OnTextChanged(const FText& InText
 	}
 }
 
+
 void SSessionLauncherDeployRepositorySettings::OnTextCommitted( const FText& InText, ETextCommit::Type CommitInfo)
 {
 	if (CommitInfo == ETextCommit::OnEnter)
@@ -148,4 +146,5 @@ void SSessionLauncherDeployRepositorySettings::OnTextCommitted( const FText& InT
 	}
 }
 
+
 #undef LOCTEXT_NAMESPACE
diff --git a/Engine/Source/Developer/SessionLauncher/Private/Widgets/Deploy/SSessionLauncherDeployRepositorySettings.h b/Engine/Source/Developer/SessionLauncher/Private/Widgets/Deploy/SSessionLauncherDeployRepositorySettings.h
--- a/Engine/Source/Developer/SessionLauncher/Private/Widgets/Deploy/SSessionLauncherDeployRepositorySettings.h
+++ b/Engine/Source/Developer/SessionLauncher/Private/Widgets/Deploy/SSessionLauncherDeployRepositorySettings.h
@@ -1,9 +1,5 @@
 // Copyright 1998-2014 Epic Games, Inc. All Rights Reserved.
 
-/*=============================================================================
-	SSessionLauncherDeployRepositorySettings.h: Declares the SSessionLauncherDeployRepositorySettings class.
-=============================================================================*/
-
 #pragma once
 
 
@@ -18,7 +14,6 @@ public:
 	SLATE_BEGIN_ARGS(SSessionLauncherDeployRepositorySettings) { }
 	SLATE_END_ARGS()
 
-
 public:
 
 	/**
@@ -29,7 +24,6 @@ public:
 	 */
 	void Construct(	const FArguments& InArgs, const FSessionLauncherModelRef& InModel );
 
-
 private:
 
 	FReply HandleBrowseButtonClicked( );
diff --git a/Engine/Source/Developer/SessionLauncher/Private/Widgets/Deploy/SSessionLauncherDeployTargets.cpp b/Engine/Source/Developer/SessionLauncher/Private/Widgets/Deploy/SSessionLauncherDeployTargets.cpp
--- a/Engine/Source/Developer/SessionLauncher/Private/Widgets/Deploy/SSessionLauncherDeployTargets.cpp
+++ b/Engine/Source/Developer/SessionLauncher/Private/Widgets/Deploy/SSessionLauncherDeployTargets.cpp
@@ -1,9 +1,5 @@
 // Copyright 1998-2014 Epic Games, Inc. All Rights Reserved.
 
-/*=============================================================================
-	SSessionLauncherDeployTargets.cpp: Implements the SSessionLauncherDeployTargets class.
-=============================================================================*/
-
 #include "SessionLauncherPrivatePCH.h"
 
 
diff --git a/Engine/Source/Developer/SessionLauncher/Private/Widgets/Deploy/SSessionLauncherDeployTargets.h b/Engine/Source/Developer/SessionLauncher/Private/Widgets/Deploy/SSessionLauncherDeployTargets.h
--- a/Engine/Source/Developer/SessionLauncher/Private/Widgets/Deploy/SSessionLauncherDeployTargets.h
+++ b/Engine/Source/Developer/SessionLauncher/Private/Widgets/Deploy/SSessionLauncherDeployTargets.h
@@ -1,24 +1,19 @@
 // Copyright 1998-2014 Epic Games, Inc. All Rights Reserved.
 
-/*=============================================================================
-	SSessionLauncherDeployTargets.h: Declares the SSessionLauncherDeployTargets class.
-=============================================================================*/
-
 #pragma once
 
 
 namespace EShowDevicesChoices
 {
+	/**
+	 * Enumerates choices for the 'Show Devices' filter.
+	 */
 	enum Type
 	{
-		/**
-		 * Show all available devices.
-		 */
+		/** Show all available devices. */
 		ShowAllDevices,
 
-		/**
-		 * Only show devices that are included in the selected device group.
-		 */
+		/** Only show devices that are included in the selected device group. */
 		ShowGroupDevices
 	};
 }
@@ -35,7 +30,6 @@ public:
 	SLATE_BEGIN_ARGS(SSessionLauncherDeployTargets) { }
 	SLATE_END_ARGS()
 
-
 public:
 
 	/**
@@ -43,26 +37,21 @@ public:
 	 */
 	~SSessionLauncherDeployTargets( );
 
-
 public:
 
 	/**
 	 * Constructs the widget.
 	 *
-	 * @param InArgs - The Slate argument list.
-	 * @param InModel - The data model.
+	 * @param InArgs The Slate argument list.
+	 * @param InModel The data model.
 	 */
 	void Construct(	const FArguments& InArgs, const FSessionLauncherModelRef& InModel, bool InIsAlwaysEnabled = false );
 
-
 public:
 
-	// Begin SCompoundWidget overrides
-
-	virtual void Tick( const FGeometry& AllottedGeometry, const double InCurrentTime, const float InDeltaTime ) OVERRIDE;
-
-	// End SCompoundWidget overrides
+	// SCompoundWidget overrides
 
+	virtual void Tick( const FGeometry& AllottedGeometry, const double InCurrentTime, const float InDeltaTime ) override;
 
 protected:
 
@@ -71,7 +60,6 @@ protected:
 	 */
 	void RefreshDeviceProxyList( );
 
-
 private:
 
 	// Callback for selecting a different device group in the device group selector.
@@ -116,7 +104,6 @@ private:
 	// Callback for determining the checked state of the specified 'Show devices' check box.
 	ESlateCheckBoxState::Type HandleShowCheckBoxIsChecked( EShowDevicesChoices::Type Choice ) const;
 
-
 private:
 
 	// Holds a pointer to the active device group.
diff --git a/Engine/Source/Developer/SessionLauncher/Private/Widgets/Deploy/SSessionLauncherDeployToDeviceSettings.cpp b/Engine/Source/Developer/SessionLauncher/Private/Widgets/Deploy/SSessionLauncherDeployToDeviceSettings.cpp
--- a/Engine/Source/Developer/SessionLauncher/Private/Widgets/Deploy/SSessionLauncherDeployToDeviceSettings.cpp
+++ b/Engine/Source/Developer/SessionLauncher/Private/Widgets/Deploy/SSessionLauncherDeployToDeviceSettings.cpp
@@ -1,9 +1,5 @@
 // Copyright 1998-2014 Epic Games, Inc. All Rights Reserved.
 
-/*=============================================================================
-	SSessionLauncherDeployToDeviceSettings.cpp: Implements the SSessionLauncherDeployToDeviceSettings class.
-=============================================================================*/
-
 #include "SessionLauncherPrivatePCH.h"
 
 
diff --git a/Engine/Source/Developer/SessionLauncher/Private/Widgets/Deploy/SSessionLauncherDeployToDeviceSettings.h b/Engine/Source/Developer/SessionLauncher/Private/Widgets/Deploy/SSessionLauncherDeployToDeviceSettings.h
--- a/Engine/Source/Developer/SessionLauncher/Private/Widgets/Deploy/SSessionLauncherDeployToDeviceSettings.h
+++ b/Engine/Source/Developer/SessionLauncher/Private/Widgets/Deploy/SSessionLauncherDeployToDeviceSettings.h
@@ -1,9 +1,5 @@
 // Copyright 1998-2014 Epic Games, Inc. All Rights Reserved.
 
-/*=============================================================================
-	SSessionLauncherDeployToDeviceSettings.h: Declares the SSessionLauncherDeployToDeviceSettings class.
-=============================================================================*/
-
 #pragma once
 
 
@@ -18,7 +14,6 @@ public:
 	SLATE_BEGIN_ARGS(SSessionLauncherDeployToDeviceSettings) { }
 	SLATE_END_ARGS()
 
-
 public:
 
 	/**
@@ -29,7 +24,6 @@ public:
 	 */
 	void Construct(	const FArguments& InArgs, const FSessionLauncherModelRef& InModel, EVisibility InShowAdvanced = EVisibility::Visible );
 
-
 private:
 
 	// Callback for check state changes of the 'UnrealPak' check box.
@@ -38,7 +32,6 @@ private:
 	// Callback for determining the checked state of the 'UnrealPak' check box.
 	ESlateCheckBoxState::Type HandleUnrealPakCheckBoxIsChecked( ) const;
 
-
 private:
 
 	// Holds a pointer to the data model.
diff --git a/Engine/Source/Developer/SessionLauncher/Private/Widgets/Deploy/SSessionLauncherDeviceGroupSelector.h b/Engine/Source/Developer/SessionLauncher/Private/Widgets/Deploy/SSessionLauncherDeviceGroupSelector.h
--- a/Engine/Source/Developer/SessionLauncher/Private/Widgets/Deploy/SSessionLauncherDeviceGroupSelector.h
+++ b/Engine/Source/Developer/SessionLauncher/Private/Widgets/Deploy/SSessionLauncherDeviceGroupSelector.h
@@ -1,9 +1,5 @@
 // Copyright 1998-2014 Epic Games, Inc. All Rights Reserved.
 
-/*=============================================================================
-	SSessionLauncherDeviceGroupSelector.h: Declares the SSessionLauncherDeviceGroupSelector class.
-=============================================================================*/
-
 #pragma once
 
 
@@ -40,7 +36,6 @@ public:
 
 	SLATE_END_ARGS()
 
-
 public:
 
 	/**
@@ -55,7 +50,6 @@ public:
 		}
 	}
 
-
 public:
 
 	/**
@@ -109,7 +103,6 @@ public:
 		ProfileManager->OnDeviceGroupRemoved().AddSP(this, &SSessionLauncherDeviceGroupSelector::HandleProfileManagerDeviceGroupsChanged);
 	}
 
-
 	/**
 	 * Gets the currently selected device group.
 	 *
@@ -137,7 +130,6 @@ public:
 		}		
 	}
 
-
 private:
 
 	// Callback for clicking the 'Add device group' button
@@ -234,7 +226,6 @@ private:
 		DeviceGroupComboBox->RefreshOptions();
 	}
 
-
 private:
 
 	// Holds the device group combo box.
@@ -243,7 +234,6 @@ private:
 	// Holds the profile manager.
 	ILauncherProfileManagerPtr ProfileManager;
 
-
 private:
 
 	// Holds a delegate to be invoked when a different device group has been selected.
diff --git a/Engine/Source/Developer/SessionLauncher/Private/Widgets/Deploy/SSessionLauncherDeviceListRow.h b/Engine/Source/Developer/SessionLauncher/Private/Widgets/Deploy/SSessionLauncherDeviceListRow.h
--- a/Engine/Source/Developer/SessionLauncher/Private/Widgets/Deploy/SSessionLauncherDeviceListRow.h
+++ b/Engine/Source/Developer/SessionLauncher/Private/Widgets/Deploy/SSessionLauncherDeviceListRow.h
@@ -1,9 +1,5 @@
 // Copyright 1998-2014 Epic Games, Inc. All Rights Reserved.
 
-/*=============================================================================
-	SSessionLauncherDevicesListRow.h: Declares the SSessionLauncherDevicesListRow class.
-=============================================================================*/
-
 #pragma once
 
 
@@ -37,7 +33,6 @@ public:
 
 	SLATE_END_ARGS()
 
-
 public:
 
 	/**
@@ -55,7 +50,6 @@ public:
 		SMultiColumnTableRow<ITargetDeviceProxyPtr>::Construct(FSuperRowType::FArguments(), InOwnerTableView);
 	}
 
-
 public:
 
 	/**
@@ -128,7 +122,6 @@ public:
 		return SNullWidget::NullWidget;
 	}
 
-
 private:
 
 	// Callback for changing this row's check box state.
diff --git a/Engine/Source/Developer/SessionLauncher/Private/Widgets/Launch/SSessionLauncherLaunchCustomRoles.cpp b/Engine/Source/Developer/SessionLauncher/Private/Widgets/Launch/SSessionLauncherLaunchCustomRoles.cpp
--- a/Engine/Source/Developer/SessionLauncher/Private/Widgets/Launch/SSessionLauncherLaunchCustomRoles.cpp
+++ b/Engine/Source/Developer/SessionLauncher/Private/Widgets/Launch/SSessionLauncherLaunchCustomRoles.cpp
@@ -1,9 +1,5 @@
 // Copyright 1998-2014 Epic Games, Inc. All Rights Reserved.
 
-/*=============================================================================
-	SSessionLauncherLaunchCustomRoles.cpp: Implements the SSessionLauncherLaunchCustomRoles class.
-=============================================================================*/
-
 #include "SessionLauncherPrivatePCH.h"
 
 
diff --git a/Engine/Source/Developer/SessionLauncher/Private/Widgets/Launch/SSessionLauncherLaunchCustomRoles.h b/Engine/Source/Developer/SessionLauncher/Private/Widgets/Launch/SSessionLauncherLaunchCustomRoles.h
--- a/Engine/Source/Developer/SessionLauncher/Private/Widgets/Launch/SSessionLauncherLaunchCustomRoles.h
+++ b/Engine/Source/Developer/SessionLauncher/Private/Widgets/Launch/SSessionLauncherLaunchCustomRoles.h
@@ -1,9 +1,5 @@
 // Copyright 1998-2014 Epic Games, Inc. All Rights Reserved.
 
-/*=============================================================================
-	SSessionLauncherLaunchCustomRoles.h: Declares the SSessionLauncherLaunchCustomRoles class.
-=============================================================================*/
-
 #pragma once
 
 
@@ -18,7 +14,6 @@ public:
 	SLATE_BEGIN_ARGS(SSessionLauncherLaunchCustomRoles) { }
 	SLATE_END_ARGS()
 
-
 public:
 
 	/**
@@ -26,7 +21,6 @@ public:
 	 */
 	~SSessionLauncherLaunchCustomRoles( );
 
-
 public:
 
 	/**
@@ -37,10 +31,6 @@ public:
 	 */
 	void Construct( const FArguments& InArgs, const FSessionLauncherModelRef& InModel );
 
-
-protected:
-
-
 private:
 
 	// Holds a pointer to the data model.
diff --git a/Engine/Source/Developer/SessionLauncher/Private/Widgets/Launch/SSessionLauncherLaunchPage.cpp b/Engine/Source/Developer/SessionLauncher/Private/Widgets/Launch/SSessionLauncherLaunchPage.cpp
--- a/Engine/Source/Developer/SessionLauncher/Private/Widgets/Launch/SSessionLauncherLaunchPage.cpp
+++ b/Engine/Source/Developer/SessionLauncher/Private/Widgets/Launch/SSessionLauncherLaunchPage.cpp
@@ -1,9 +1,5 @@
 // Copyright 1998-2014 Epic Games, Inc. All Rights Reserved.
 
-/*=============================================================================
-	SSessionLauncherLaunchPage.cpp: Implements the SSessionLauncherLaunchPage class.
-=============================================================================*/
-
 #include "SessionLauncherPrivatePCH.h"
 
 
diff --git a/Engine/Source/Developer/SessionLauncher/Private/Widgets/Launch/SSessionLauncherLaunchPage.h b/Engine/Source/Developer/SessionLauncher/Private/Widgets/Launch/SSessionLauncherLaunchPage.h
--- a/Engine/Source/Developer/SessionLauncher/Private/Widgets/Launch/SSessionLauncherLaunchPage.h
+++ b/Engine/Source/Developer/SessionLauncher/Private/Widgets/Launch/SSessionLauncherLaunchPage.h
@@ -1,9 +1,5 @@
 // Copyright 1998-2014 Epic Games, Inc. All Rights Reserved.
 
-/*=============================================================================
-	SSessionLauncherLaunchPage.h: Declares the SSessionLauncherLaunchPage class.
-=============================================================================*/
-
 #pragma once
 
 
@@ -18,7 +14,6 @@ public:
 	SLATE_BEGIN_ARGS(SSessionLauncherLaunchPage) { }
 	SLATE_END_ARGS()
 
-
 public:
 
 	/**
@@ -26,14 +21,13 @@ public:
 	 */
 	~SSessionLauncherLaunchPage( );
 
-
 public:
 
 	/**
 	 * Constructs the widget.
 	 *
-	 * @param InArgs - The Slate argument list.
-	 * @param InModel - The data model.
+	 * @param InArgs The Slate argument list.
+	 * @param InModel The data model.
 	 */
 	void Construct(	const FArguments& InArgs, const FSessionLauncherModelRef& InModel );
 
@@ -42,7 +36,6 @@ public:
 	 */
 	void Refresh( );
 
-
 private:
 
 	// Callback for getting the visibility of the 'Cannot launch' text block.
diff --git a/Engine/Source/Developer/SessionLauncher/Private/Widgets/Launch/SSessionLauncherLaunchRoleEditor.h b/Engine/Source/Developer/SessionLauncher/Private/Widgets/Launch/SSessionLauncherLaunchRoleEditor.h
--- a/Engine/Source/Developer/SessionLauncher/Private/Widgets/Launch/SSessionLauncherLaunchRoleEditor.h
+++ b/Engine/Source/Developer/SessionLauncher/Private/Widgets/Launch/SSessionLauncherLaunchRoleEditor.h
@@ -1,9 +1,5 @@
 // Copyright 1998-2014 Epic Games, Inc. All Rights Reserved.
 
-/*=============================================================================
-	SSessionLauncherLaunchRoleEditor.h: Declares the SSessionLauncherLaunchRoleEditor class.
-=============================================================================*/
-
 #pragma once
 
 
@@ -23,31 +19,24 @@ public:
 		, _AvailableMaps()
 	{ }
 		
-		/**
-		 * The role to be edited initially.
-		 */
+		/** The role to be edited initially. */
 		SLATE_ARGUMENT(ILauncherProfileLaunchRolePtr, InitialRole)
 
-		/**
-		 * The list of available cultures.
-		 */
+		/** The list of available cultures. */
 		SLATE_ARGUMENT(const TArray<FString>*, AvailableCultures)
 
-		/**
-		 * The list of available maps.
-		 */
+		/** The list of available maps. */
 		SLATE_ARGUMENT(const TArray<FString>*, AvailableMaps)
 
 	SLATE_END_ARGS()
 
-
 public:
 
 	/**
 	 * Constructs the widget.
 	 *
-	 * @param InArgs - The Slate argument list.
-	 * @param InRole - The launch role to edit.
+	 * @param InArgs The Slate argument list.
+	 * @param InRole The launch role to edit.
 	 */
 	BEGIN_SLATE_FUNCTION_BUILD_OPTIMIZATION
 	void Construct(	const FArguments& InArgs )
@@ -184,7 +173,7 @@ public:
 	/**
 	 * Refreshes the widget.
 	 *
-	 * @param InRole - The role to edit, or NULL if no role is being edited.
+	 * @param InRole The role to edit, or NULL if no role is being edited.
 	 */
 	void Refresh( const ILauncherProfileLaunchRolePtr& InRole )
 	{
@@ -256,7 +245,6 @@ public:
 		MapComboBox->RefreshOptions();
 	}
 
-
 private:
 
 	// Callback for when the text of the command line text box has changed.
@@ -435,7 +423,6 @@ private:
 		return ESlateCheckBoxState::Unchecked;
 	}
 
-
 private:
 
 	// Holds a pointer to the list of available cultures.
diff --git a/Engine/Source/Developer/SessionLauncher/Private/Widgets/Package/SSessionLauncherPackagePage.cpp b/Engine/Source/Developer/SessionLauncher/Private/Widgets/Package/SSessionLauncherPackagePage.cpp
--- a/Engine/Source/Developer/SessionLauncher/Private/Widgets/Package/SSessionLauncherPackagePage.cpp
+++ b/Engine/Source/Developer/SessionLauncher/Private/Widgets/Package/SSessionLauncherPackagePage.cpp
@@ -1,9 +1,5 @@
 // Copyright 1998-2014 Epic Games, Inc. All Rights Reserved.
 
-/*=============================================================================
-	SSessionLauncherPackagePage.cpp: Implements the SSessionLauncherPackagePage class.
-=============================================================================*/
-
 #include "SessionLauncherPrivatePCH.h"
 
 
diff --git a/Engine/Source/Developer/SessionLauncher/Private/Widgets/Package/SSessionLauncherPackagePage.h b/Engine/Source/Developer/SessionLauncher/Private/Widgets/Package/SSessionLauncherPackagePage.h
--- a/Engine/Source/Developer/SessionLauncher/Private/Widgets/Package/SSessionLauncherPackagePage.h
+++ b/Engine/Source/Developer/SessionLauncher/Private/Widgets/Package/SSessionLauncherPackagePage.h
@@ -1,9 +1,5 @@
 // Copyright 1998-2014 Epic Games, Inc. All Rights Reserved.
 
-/*=============================================================================
-	SSessionLauncherPackagePage.h: Declares the SSessionLauncherPackagePage class.
-=============================================================================*/
-
 #pragma once
 
 
@@ -18,7 +14,6 @@ public:
 	SLATE_BEGIN_ARGS(SSessionLauncherPackagePage) { }
 	SLATE_END_ARGS()
 
-
 public:
 
 	/**
@@ -26,18 +21,16 @@ public:
 	 */
 	~SSessionLauncherPackagePage( );
 
-
 public:
 
 	/**
 	 * Constructs the widget.
 	 *
-	 * @param InArgs - The Slate argument list.
-	 * @param InModel - The data model.
+	 * @param InArgs The Slate argument list.
+	 * @param InModel The data model.
 	 */
 	void Construct( const FArguments& InArgs, const FSessionLauncherModelRef& InModel );
 
-
 private:
 
 	// Callback for getting the content text of the 'Cook Mode' combo button.
@@ -52,7 +45,6 @@ private:
 	// Callback for changing the selected profile in the profile manager.
 	void HandleProfileManagerProfileSelected( const ILauncherProfilePtr& SelectedProfile, const ILauncherProfilePtr& PreviousProfile );
 
-
 private:
 
 	// Holds a pointer to the data model.
diff --git a/Engine/Source/Developer/SessionLauncher/Private/Widgets/Package/SSessionLauncherPackagingSettings.cpp b/Engine/Source/Developer/SessionLauncher/Private/Widgets/Package/SSessionLauncherPackagingSettings.cpp
--- a/Engine/Source/Developer/SessionLauncher/Private/Widgets/Package/SSessionLauncherPackagingSettings.cpp
+++ b/Engine/Source/Developer/SessionLauncher/Private/Widgets/Package/SSessionLauncherPackagingSettings.cpp
@@ -1,9 +1,5 @@
 // Copyright 1998-2014 Epic Games, Inc. All Rights Reserved.
 
-/*=============================================================================
-	SSessionLauncherPackagingSettings.cpp: Implements the SSessionLauncherPackagingSettings class.
-=============================================================================*/
-
 #include "SessionLauncherPrivatePCH.h"
 
 
diff --git a/Engine/Source/Developer/SessionLauncher/Private/Widgets/Package/SSessionLauncherPackagingSettings.h b/Engine/Source/Developer/SessionLauncher/Private/Widgets/Package/SSessionLauncherPackagingSettings.h
--- a/Engine/Source/Developer/SessionLauncher/Private/Widgets/Package/SSessionLauncherPackagingSettings.h
+++ b/Engine/Source/Developer/SessionLauncher/Private/Widgets/Package/SSessionLauncherPackagingSettings.h
@@ -1,9 +1,5 @@
 // Copyright 1998-2014 Epic Games, Inc. All Rights Reserved.
 
-/*=============================================================================
-	SSessionLauncherPackagingSettings.h: Declares the SSessionLauncherPackagingSettings class.
-=============================================================================*/
-
 #pragma once
 
 
@@ -18,7 +14,6 @@ public:
 	SLATE_BEGIN_ARGS(SSessionLauncherPackagingSettings) { }
 	SLATE_END_ARGS()
 
-
 public:
 
 	/**
@@ -26,24 +21,21 @@ public:
 	 */
 	~SSessionLauncherPackagingSettings( );
 
-
 public:
 
 	/**
 	 * Constructs the widget.
 	 *
-	 * @param InArgs - The Slate argument list.
-	 * @param InModel - The data model.
+	 * @param InArgs The Slate argument list.
+	 * @param InModel The data model.
 	 */
 	void Construct( const FArguments& InArgs, const FSessionLauncherModelRef& InModel );
 
-
 private:
 
 	// Callback for changing the selected profile in the profile manager.
 	void HandleProfileManagerProfileSelected( const ILauncherProfilePtr& SelectedProfile, const ILauncherProfilePtr& PreviousProfile );
 
-
 private:
 
 	// Holds a pointer to the data model.
diff --git a/Engine/Source/Developer/SessionLauncher/Private/Widgets/Project/SSessionLauncherProjectPage.cpp b/Engine/Source/Developer/SessionLauncher/Private/Widgets/Project/SSessionLauncherProjectPage.cpp
--- a/Engine/Source/Developer/SessionLauncher/Private/Widgets/Project/SSessionLauncherProjectPage.cpp
+++ b/Engine/Source/Developer/SessionLauncher/Private/Widgets/Project/SSessionLauncherProjectPage.cpp
@@ -1,9 +1,5 @@
 // Copyright 1998-2014 Epic Games, Inc. All Rights Reserved.
 
-/*=============================================================================
-	SSessionLauncherProjectPage.cpp: Implements the SSessionLauncherProjectPage class.
-=============================================================================*/
-
 #include "SessionLauncherPrivatePCH.h"
 
 
diff --git a/Engine/Source/Developer/SessionLauncher/Private/Widgets/Project/SSessionLauncherProjectPage.h b/Engine/Source/Developer/SessionLauncher/Private/Widgets/Project/SSessionLauncherProjectPage.h
--- a/Engine/Source/Developer/SessionLauncher/Private/Widgets/Project/SSessionLauncherProjectPage.h
+++ b/Engine/Source/Developer/SessionLauncher/Private/Widgets/Project/SSessionLauncherProjectPage.h
@@ -1,9 +1,5 @@
 // Copyright 1998-2014 Epic Games, Inc. All Rights Reserved.
 
-/*=============================================================================
-	SSessionLauncherProjectPage.h: Declares the SSessionLauncherProjectPage class.
-=============================================================================*/
-
 #pragma once
 
 
@@ -18,24 +14,21 @@ public:
 	SLATE_BEGIN_ARGS(SSessionLauncherProjectPage) { }
 	SLATE_END_ARGS()
 
-
 public:
 
 	/**
 	 * Constructs the widget.
 	 *
-	 * @param InArgs - The Slate argument list.
-	 * @param InModel - The data model.
+	 * @param InArgs The Slate argument list.
+	 * @param InModel The data model.
 	 */
 	void Construct(	const FArguments& InArgs, const FSessionLauncherModelRef& InModel, bool InShowConfig = true );
 
-
 private:
 
 	// Callback for getting the enabled state of the project settings area.
 	bool HandleProjectSettingsAreaIsEnabled( ) const;
 
-
 private:
 
 	// Holds a pointer to the data model.
diff --git a/Engine/Source/Developer/SessionLauncher/Private/Widgets/Project/SSessionLauncherProjectPicker.cpp b/Engine/Source/Developer/SessionLauncher/Private/Widgets/Project/SSessionLauncherProjectPicker.cpp
--- a/Engine/Source/Developer/SessionLauncher/Private/Widgets/Project/SSessionLauncherProjectPicker.cpp
+++ b/Engine/Source/Developer/SessionLauncher/Private/Widgets/Project/SSessionLauncherProjectPicker.cpp
@@ -1,9 +1,5 @@
 // Copyright 1998-2014 Epic Games, Inc. All Rights Reserved.
 
-/*=============================================================================
-	SSessionLauncherUnrealProjectPicker.cpp: Implements the SSessionLauncherUnrealProjectPicker class.
-=============================================================================*/
-
 #include "SessionLauncherPrivatePCH.h"
 
 
@@ -63,118 +59,121 @@ TSharedRef<SWidget> SSessionLauncherProjectPicker::MakeProjectMenuWidget( )
 	return MenuBuilder.MakeWidget();
 }
 
+
 TSharedRef<SWidget> SSessionLauncherProjectPicker::MakeProjectWidget()
 {
 	TSharedRef<SWidget> Widget = SNew(SVerticalBox)
 
-		+ SVerticalBox::Slot()
+	+ SVerticalBox::Slot()
 		.FillHeight(0.5)
 		[
 			SNew(SBorder)
-			.BorderImage(FEditorStyle::GetBrush("ToolPanel.GroupBorder"))
-			.Padding(8.0)
-			[
-				SNew(SVerticalBox)
-
-				+ SVerticalBox::Slot()
-				.AutoHeight()
+				.BorderImage(FEditorStyle::GetBrush("ToolPanel.GroupBorder"))
+				.Padding(8.0)
 				[
-					SNew(SSessionLauncherFormLabel)
-					.ErrorToolTipText(NSLOCTEXT("SSessionLauncherBuildValidation", "NoProjectSelectedError", "A Project must be selected."))
-					.ErrorVisibility(this, &SSessionLauncherProjectPicker::HandleValidationErrorIconVisibility, ELauncherProfileValidationErrors::NoProjectSelected)
-					.LabelText(LOCTEXT("ProjectComboBoxLabel", "Project:"))
-				]
+					SNew(SVerticalBox)
 
-				+ SVerticalBox::Slot()
-					.AutoHeight()
-					.Padding(0.0f, 4.0f, 0.0f, 0.0f)
-					[
-						// project selector
-						SNew(SComboButton)
-						.ButtonContent()
+					+ SVerticalBox::Slot()
+						.AutoHeight()
 						[
-							SNew(STextBlock)
-							.Text(this, &SSessionLauncherProjectPicker::HandleProjectComboButtonText)
+							SNew(SSessionLauncherFormLabel)
+								.ErrorToolTipText(NSLOCTEXT("SSessionLauncherBuildValidation", "NoProjectSelectedError", "A Project must be selected."))
+								.ErrorVisibility(this, &SSessionLauncherProjectPicker::HandleValidationErrorIconVisibility, ELauncherProfileValidationErrors::NoProjectSelected)
+								.LabelText(LOCTEXT("ProjectComboBoxLabel", "Project:"))
 						]
-						.ContentPadding(FMargin(6.0f, 2.0f))
-							.MenuContent()
-							[
-								MakeProjectMenuWidget()
-							]
-						.ToolTipText(this, &SSessionLauncherProjectPicker::HandleProjectComboButtonToolTip)
-					]
-			]
+
+					+ SVerticalBox::Slot()
+						.AutoHeight()
+						.Padding(0.0f, 4.0f, 0.0f, 0.0f)
+						[
+							// project selector
+							SNew(SComboButton)
+								.ButtonContent()
+								[
+									SNew(STextBlock)
+										.Text(this, &SSessionLauncherProjectPicker::HandleProjectComboButtonText)
+								]
+								.ContentPadding(FMargin(6.0f, 2.0f))
+								.MenuContent()
+								[
+									MakeProjectMenuWidget()
+								]
+								.ToolTipText(this, &SSessionLauncherProjectPicker::HandleProjectComboButtonToolTip)
+						]
+				]
 		];
 
 	return Widget;
 }
 
+
 TSharedRef<SWidget> SSessionLauncherProjectPicker::MakeProjectAndConfigWidget()
 {
 	TSharedRef<SWidget> Widget = SNew(SVerticalBox)
 
-		+ SVerticalBox::Slot()
+	+ SVerticalBox::Slot()
 		.FillHeight(0.5)
 		[
 			SNew(SBorder)
-			.BorderImage(FEditorStyle::GetBrush("ToolPanel.GroupBorder"))
-			.Padding(8.0)
-			[
-				SNew(SVerticalBox)
-
-				+ SVerticalBox::Slot()
-				.AutoHeight()
+				.BorderImage(FEditorStyle::GetBrush("ToolPanel.GroupBorder"))
+				.Padding(8.0)
 				[
-					SNew(SSessionLauncherFormLabel)
-					.ErrorToolTipText(NSLOCTEXT("SSessionLauncherBuildValidation", "NoProjectSelectedError", "A Project must be selected."))
-					.ErrorVisibility(this, &SSessionLauncherProjectPicker::HandleValidationErrorIconVisibility, ELauncherProfileValidationErrors::NoProjectSelected)
-					.LabelText(LOCTEXT("ProjectComboBoxLabel", "Project:"))
-				]
+					SNew(SVerticalBox)
 
-				+ SVerticalBox::Slot()
-					.AutoHeight()
-					.Padding(0.0f, 4.0f, 0.0f, 0.0f)
-					[
-						// project selector
-						SNew(SComboButton)
-						.ButtonContent()
+					+ SVerticalBox::Slot()
+						.AutoHeight()
 						[
-							SNew(STextBlock)
-							.Text(this, &SSessionLauncherProjectPicker::HandleProjectComboButtonText)
+							SNew(SSessionLauncherFormLabel)
+								.ErrorToolTipText(NSLOCTEXT("SSessionLauncherBuildValidation", "NoProjectSelectedError", "A Project must be selected."))
+								.ErrorVisibility(this, &SSessionLauncherProjectPicker::HandleValidationErrorIconVisibility, ELauncherProfileValidationErrors::NoProjectSelected)
+								.LabelText(LOCTEXT("ProjectComboBoxLabel", "Project:"))
 						]
-						.ContentPadding(FMargin(6.0f, 2.0f))
-							.MenuContent()
-							[
-								MakeProjectMenuWidget()
-							]
-						.ToolTipText(this, &SSessionLauncherProjectPicker::HandleProjectComboButtonToolTip)
-					]
-
-				+ SVerticalBox::Slot()
-					.AutoHeight()
-					.Padding(0.0f, 8.0f, 0.0f, 0.0f)
-					[
-						SNew(SSessionLauncherFormLabel)
-						.ErrorToolTipText(NSLOCTEXT("SSessionLauncherBuildValidation", "NoBuildConfigurationSelectedError", "A Build Configuration must be selected."))
-						.ErrorVisibility(this, &SSessionLauncherProjectPicker::HandleValidationErrorIconVisibility, ELauncherProfileValidationErrors::NoBuildConfigurationSelected)
-						.LabelText(LOCTEXT("ConfigurationComboBoxLabel", "Build Configuration:"))
-					]
-
-				+ SVerticalBox::Slot()
-					.AutoHeight()
-					.Padding(0.0f, 4.0f, 0.0f, 0.0f)
-					[
-						// build configuration selector
-						SNew(SSessionLauncherBuildConfigurationSelector)
-						.OnConfigurationSelected(this, &SSessionLauncherProjectPicker::HandleBuildConfigurationSelectorConfigurationSelected)
-						.Text(this, &SSessionLauncherProjectPicker::HandleBuildConfigurationSelectorText)
-					]
-			]
+
+					+ SVerticalBox::Slot()
+						.AutoHeight()
+						.Padding(0.0f, 4.0f, 0.0f, 0.0f)
+						[
+							// project selector
+							SNew(SComboButton)
+								.ButtonContent()
+								[
+									SNew(STextBlock)
+										.Text(this, &SSessionLauncherProjectPicker::HandleProjectComboButtonText)
+								]
+								.ContentPadding(FMargin(6.0f, 2.0f))
+								.MenuContent()
+								[
+									MakeProjectMenuWidget()
+								]
+								.ToolTipText(this, &SSessionLauncherProjectPicker::HandleProjectComboButtonToolTip)
+						]
+
+					+ SVerticalBox::Slot()
+						.AutoHeight()
+						.Padding(0.0f, 8.0f, 0.0f, 0.0f)
+						[
+							SNew(SSessionLauncherFormLabel)
+								.ErrorToolTipText(NSLOCTEXT("SSessionLauncherBuildValidation", "NoBuildConfigurationSelectedError", "A Build Configuration must be selected."))
+								.ErrorVisibility(this, &SSessionLauncherProjectPicker::HandleValidationErrorIconVisibility, ELauncherProfileValidationErrors::NoBuildConfigurationSelected)
+								.LabelText(LOCTEXT("ConfigurationComboBoxLabel", "Build Configuration:"))
+						]
+
+					+ SVerticalBox::Slot()
+						.AutoHeight()
+						.Padding(0.0f, 4.0f, 0.0f, 0.0f)
+						[
+							// build configuration selector
+							SNew(SSessionLauncherBuildConfigurationSelector)
+								.OnConfigurationSelected(this, &SSessionLauncherProjectPicker::HandleBuildConfigurationSelectorConfigurationSelected)
+								.Text(this, &SSessionLauncherProjectPicker::HandleBuildConfigurationSelectorText)
+						]
+				]
 		];
 
 	return Widget;
 }
 
+
 /* SSessionLauncherProjectPicker callbacks
  *****************************************************************************/
 
diff --git a/Engine/Source/Developer/SessionLauncher/Private/Widgets/Project/SSessionLauncherProjectPicker.h b/Engine/Source/Developer/SessionLauncher/Private/Widgets/Project/SSessionLauncherProjectPicker.h
--- a/Engine/Source/Developer/SessionLauncher/Private/Widgets/Project/SSessionLauncherProjectPicker.h
+++ b/Engine/Source/Developer/SessionLauncher/Private/Widgets/Project/SSessionLauncherProjectPicker.h
@@ -1,9 +1,5 @@
 // Copyright 1998-2014 Epic Games, Inc. All Rights Reserved.
 
-/*=============================================================================
-	SSessionLauncherProjectPicker.h: Declares the SSessionLauncherProjectPicker class.
-=============================================================================*/
-
 #pragma once
 
 
@@ -18,7 +14,6 @@ public:
 	SLATE_BEGIN_ARGS(SSessionLauncherProjectPicker) { }
 	SLATE_END_ARGS()
 
-
 public:
 
 	/**
@@ -26,18 +21,16 @@ public:
 	 */
 	~SSessionLauncherProjectPicker( );
 
-
 public:
 
 	/**
 	 * Constructs the widget.
 	 *
-	 * @param InArgs - The Slate argument list.
-	 * @param InModel - The data model.
+	 * @param InArgs The Slate argument list.
+	 * @param InModel The data model.
 	 */
 	void Construct(	const FArguments& InArgs, const FSessionLauncherModelRef& InModel, bool InShowConfiguration = true );
 
-
 protected:
 
 	/**
@@ -87,7 +80,7 @@ private:
 private:
 
 	// Holds the list of available projects.
-	TArray<TSharedPtr<FString> > ProjectList;
+	TArray<TSharedPtr<FString>> ProjectList;
 
 	// Holds a pointer to the data model.
 	FSessionLauncherModelPtr Model;
