Docs: Removed file comments and added missing code documentation

diff --git a/Engine/Source/Developer/SessionLauncher/Private/Widgets/SSessionLauncher.cpp b/Engine/Source/Developer/SessionLauncher/Private/Widgets/SSessionLauncher.cpp
--- a/Engine/Source/Developer/SessionLauncher/Private/Widgets/SSessionLauncher.cpp
+++ b/Engine/Source/Developer/SessionLauncher/Private/Widgets/SSessionLauncher.cpp
@@ -1,14 +1,11 @@
 // Copyright 1998-2014 Epic Games, Inc. All Rights Reserved.
 
-/*=============================================================================
-	SSessionLauncher.cpp: Implements the SSessionLauncher class.
-=============================================================================*/
-
 #include "SessionLauncherPrivatePCH.h"
 
 
 #define LOCTEXT_NAMESPACE "SSessionLauncher"
 
+
 /* SSessionLauncher structors
  *****************************************************************************/
 
@@ -57,218 +54,219 @@ void SSessionLauncher::Construct( const FArguments& InArgs, const TSharedRef<SDo
 		+ SVerticalBox::Slot()
 			.FillHeight(1.0f)
 			[
-		SNew(SHorizontalBox)
+				SNew(SHorizontalBox)
 
-		+ SHorizontalBox::Slot()
-			.AutoWidth()
-			[
-				SNew(SBorder)
-					.Padding(FMargin(8.0f, 6.0f, 8.0f, 4.0f))
-					.BorderImage(FEditorStyle::GetBrush("ToolPanel.GroupBorder"))
+				+ SHorizontalBox::Slot()
+					.AutoWidth()
 					[
-						// common launcher task bar
-						SAssignNew(Toolbar, SSessionLauncherToolbar, InModel, this)
+						SNew(SBorder)
+							.Padding(FMargin(8.0f, 6.0f, 8.0f, 4.0f))
+							.BorderImage(FEditorStyle::GetBrush("ToolPanel.GroupBorder"))
+							[
+								// common launcher task bar
+								SAssignNew(Toolbar, SSessionLauncherToolbar, InModel, this)
+							]
 					]
-			]
-
-		+ SHorizontalBox::Slot()
-			.FillWidth(1.0f)
-			[
-				SNew(SVerticalBox)
 
-				+ SVerticalBox::Slot()
-					.FillHeight(1.0f)
-					.Padding(8.0f, 8.0f, 0.0f, 8.0f)
+				+ SHorizontalBox::Slot()
+					.FillWidth(1.0f)
 					[
-						SAssignNew(WidgetSwitcher, SWidgetSwitcher)
-							.WidgetIndex((int32)ESessionLauncherTasks::QuickLaunch)
+						SNew(SVerticalBox)
 
-						+ SWidgetSwitcher::Slot()
+						+ SVerticalBox::Slot()
+							.FillHeight(1.0f)
+							.Padding(8.0f, 8.0f, 0.0f, 8.0f)
 							[
-								// empty panel
-								SNew(SBorder)
-							]
+								SAssignNew(WidgetSwitcher, SWidgetSwitcher)
+									.WidgetIndex((int32)ESessionLauncherTasks::QuickLaunch)
 
-						+ SWidgetSwitcher::Slot()
-							[
-								// deploy build panel
-								SNew(SSessionLauncherBuildTaskSettings, InModel)
-							]
+								+ SWidgetSwitcher::Slot()
+									[
+										// empty panel
+										SNew(SBorder)
+									]
 
-						+ SWidgetSwitcher::Slot()
-							[
-								// deploy build panel
-								SNew(SSessionLauncherDeployTaskSettings, InModel)
-							]
+								+ SWidgetSwitcher::Slot()
+									[
+										// deploy build panel
+										SNew(SSessionLauncherBuildTaskSettings, InModel)
+									]
 
-						+ SWidgetSwitcher::Slot()
-							[
-								// Quick launch panel
-								SNew(SSessionLauncherLaunchTaskSettings, InModel)
-							]
+								+ SWidgetSwitcher::Slot()
+									[
+										// deploy build panel
+										SNew(SSessionLauncherDeployTaskSettings, InModel)
+									]
 
-						+ SWidgetSwitcher::Slot()
-							[
-								// advanced settings panel
-								SNew(SSessionLauncherSettings, InModel)
-							]
+								+ SWidgetSwitcher::Slot()
+									[
+										// Quick launch panel
+										SNew(SSessionLauncherLaunchTaskSettings, InModel)
+									]
 
-						+ SWidgetSwitcher::Slot()
-							[
-								// preview panel
-								SNew(SSessionLauncherPreviewPage, InModel)
+								+ SWidgetSwitcher::Slot()
+									[
+										// advanced settings panel
+										SNew(SSessionLauncherSettings, InModel)
+									]
+
+								+ SWidgetSwitcher::Slot()
+									[
+										// preview panel
+										SNew(SSessionLauncherPreviewPage, InModel)
+									]
+
+								+ SWidgetSwitcher::Slot()
+									.Padding(8.0f, 0.0f)
+									[
+										// progress panel
+										SAssignNew(ProgressPanel, SSessionLauncherProgress)
+									]
 							]
 
-						+ SWidgetSwitcher::Slot()
-							.Padding(8.0f, 0.0f)
+						+ SVerticalBox::Slot()
+							.AutoHeight()
+							.Padding(8.0f)
 							[
-								// progress panel
-								SAssignNew(ProgressPanel, SSessionLauncherProgress)
+								SNew(SSessionLauncherValidation, InModel)
 							]
-					]
-
-				+ SVerticalBox::Slot()
-					.AutoHeight()
-					.Padding(8.0f)
-					[
-						SNew(SSessionLauncherValidation, InModel)
-					]
 
-				+ SVerticalBox::Slot()
-					.AutoHeight()
-					.Padding(0.0f, 8.0f, 0.0f, 0.0f)
-					[
-						SNew(SBorder)
-							.Padding(FMargin(8.0f, 6.0f, 8.0f, 4.0f))
-							.BorderImage(FEditorStyle::GetBrush("ToolPanel.GroupBorder"))
+						+ SVerticalBox::Slot()
+							.AutoHeight()
+							.Padding(0.0f, 8.0f, 0.0f, 0.0f)
 							[
-								SNew(SHorizontalBox)
-
-								+ SHorizontalBox::Slot()
-									.FillWidth(1.0f)
-									.HAlign(HAlign_Left)
+								SNew(SBorder)
+									.Padding(FMargin(8.0f, 6.0f, 8.0f, 4.0f))
+									.BorderImage(FEditorStyle::GetBrush("ToolPanel.GroupBorder"))
 									[
-										SNew(SOverlay)
+										SNew(SHorizontalBox)
 
-										+ SOverlay::Slot()
+										+ SHorizontalBox::Slot()
+											.FillWidth(1.0f)
+											.HAlign(HAlign_Left)
 											[
-												// edit button
-												SNew(SButton)
-													.ContentPadding(FMargin(6.0f, 2.0f))
-													.OnClicked(this, &SSessionLauncher::HandleEditButtonClicked)
-													.ToolTipText(LOCTEXT("EditButtonTooltip", "Edit the settings of the selected profile."))
-													[
-														SNew(STextBlock)
-															.Text(LOCTEXT("EditButtonLabel", "Edit Settings"))
-													]
-													.Visibility(this, &SSessionLauncher::HandleEditButtonVisibility)
-											]
+												SNew(SOverlay)
 
-										+ SOverlay::Slot()
-											[
-												// preview button
-												SNew(SButton)
-													.ContentPadding(FMargin(6.0f, 2.0f))
-													.OnClicked(this, &SSessionLauncher::HandlePreviewButtonClicked)
-													.ToolTipText(LOCTEXT("PreviewButtonTooltip", "Preview the settings of the selected profile."))
+												+ SOverlay::Slot()
 													[
-														SNew(STextBlock)
-															.Text(LOCTEXT("PreviewButtonLabel", "Preview Settings"))
+														// edit button
+														SNew(SButton)
+															.ContentPadding(FMargin(6.0f, 2.0f))
+															.OnClicked(this, &SSessionLauncher::HandleEditButtonClicked)
+															.ToolTipText(LOCTEXT("EditButtonTooltip", "Edit the settings of the selected profile."))
+															[
+																SNew(STextBlock)
+																	.Text(LOCTEXT("EditButtonLabel", "Edit Settings"))
+															]
+															.Visibility(this, &SSessionLauncher::HandleEditButtonVisibility)
 													]
-													.Visibility(this, &SSessionLauncher::HandlePreviewButtonVisibility)
-											]
-									]
-
-								+ SHorizontalBox::Slot()
-									.FillWidth(1.0f)
-									.HAlign(HAlign_Right)
-									[
-										SNew(SOverlay)
 
-										+ SOverlay::Slot()
-											[
-												// launch button
-												SNew(SButton)
-													.ContentPadding(FMargin(6.0f, 2.0f))
-													.IsEnabled(this, &SSessionLauncher::HandleFinishButtonIsEnabled)
-													.OnClicked(this, &SSessionLauncher::HandleFinishButtonClicked)
-													.ToolTipText(LOCTEXT("FinishButtonTooltip", "Launch the task as configured in the selected profile."))
+												+ SOverlay::Slot()
 													[
-														SNew(STextBlock)
-															.Text(LOCTEXT("FinishButtonLabel", "Launch"))
+														// preview button
+														SNew(SButton)
+															.ContentPadding(FMargin(6.0f, 2.0f))
+															.OnClicked(this, &SSessionLauncher::HandlePreviewButtonClicked)
+															.ToolTipText(LOCTEXT("PreviewButtonTooltip", "Preview the settings of the selected profile."))
+															[
+																SNew(STextBlock)
+																	.Text(LOCTEXT("PreviewButtonLabel", "Preview Settings"))
+															]
+															.Visibility(this, &SSessionLauncher::HandlePreviewButtonVisibility)
 													]
-													.Visibility(this, &SSessionLauncher::HandleFinishButtonVisibility)
 											]
 
-										+ SOverlay::Slot()
+										+ SHorizontalBox::Slot()
+											.FillWidth(1.0f)
+											.HAlign(HAlign_Right)
 											[
-												// cancel button
-												SNew(SButton)
-													.ContentPadding(FMargin(6.0f, 2.0f))
-													.OnClicked(this, &SSessionLauncher::HandleCancelButtonClicked)
-													.ToolTipText(LOCTEXT("CancelButtonTooltip", "Cancel the task."))
+												SNew(SOverlay)
+
+												+ SOverlay::Slot()
 													[
-														SNew(STextBlock)
-															.Text(LOCTEXT("CancelButtonLabel", "Cancel"))
+														// launch button
+														SNew(SButton)
+															.ContentPadding(FMargin(6.0f, 2.0f))
+															.IsEnabled(this, &SSessionLauncher::HandleFinishButtonIsEnabled)
+															.OnClicked(this, &SSessionLauncher::HandleFinishButtonClicked)
+															.ToolTipText(LOCTEXT("FinishButtonTooltip", "Launch the task as configured in the selected profile."))
+															[
+																SNew(STextBlock)
+																	.Text(LOCTEXT("FinishButtonLabel", "Launch"))
+															]
+															.Visibility(this, &SSessionLauncher::HandleFinishButtonVisibility)
 													]
-													.Visibility(this, &SSessionLauncher::HandleCancelButtonVisibility)
-											]
 
-										+ SOverlay::Slot()
-											[
-												SNew(SHorizontalBox)
-													.Visibility(this, &SSessionLauncher::HandleCancelingThrobberVisibility)
-
-												+ SHorizontalBox::Slot()
-													.AutoWidth()
+												+ SOverlay::Slot()
 													[
-														SNew(SCircularThrobber)
-															.NumPieces(5)
-															.Radius(10.0f)
+														// cancel button
+														SNew(SButton)
+															.ContentPadding(FMargin(6.0f, 2.0f))
+															.OnClicked(this, &SSessionLauncher::HandleCancelButtonClicked)
+															.ToolTipText(LOCTEXT("CancelButtonTooltip", "Cancel the task."))
+															[
+																SNew(STextBlock)
+																	.Text(LOCTEXT("CancelButtonLabel", "Cancel"))
+															]
+															.Visibility(this, &SSessionLauncher::HandleCancelButtonVisibility)
 													]
 
-												+ SHorizontalBox::Slot()
-													.AutoWidth()
-													.Padding(4.0f, 0.0f, 0.0f, 0.0f)
-													.VAlign(VAlign_Center)
+												+ SOverlay::Slot()
 													[
-														SNew(STextBlock)
-															.Text(LOCTEXT("CancellingText", "Canceling..."))
-													]
-											]
+														SNew(SHorizontalBox)
+															.Visibility(this, &SSessionLauncher::HandleCancelingThrobberVisibility)
 
-										+ SOverlay::Slot()
-											[
-												SNew(SHorizontalBox)
-													.Visibility(this, &SSessionLauncher::HandleDoneButtonVisibility)
+														+ SHorizontalBox::Slot()
+															.AutoWidth()
+															[
+																SNew(SCircularThrobber)
+																	.NumPieces(5)
+																	.Radius(10.0f)
+															]
 
-												+ SHorizontalBox::Slot()
-													.AutoWidth()
-													.Padding(4.0f, 0.0f, 0.0f, 0.0f)
-													[
-														// re-run button
-														SNew(SButton)
-															.ContentPadding(FMargin(6.0f, 2.0f))
-															.OnClicked(this, &SSessionLauncher::HandleFinishButtonClicked)
-															.ToolTipText(LOCTEXT("RerunButtonTooltip", "Execute the selected profile one more time."))
+														+ SHorizontalBox::Slot()
+															.AutoWidth()
+															.Padding(4.0f, 0.0f, 0.0f, 0.0f)
+															.VAlign(VAlign_Center)
 															[
 																SNew(STextBlock)
-																	.Text(LOCTEXT("RerunButtonLabel", "Run Again"))
-															]								
+																	.Text(LOCTEXT("CancellingText", "Canceling..."))
+															]
 													]
 
-												+ SHorizontalBox::Slot()
-													.AutoWidth()
-													.Padding(4.0f, 0.0f, 0.0f, 0.0f)
+												+ SOverlay::Slot()
 													[
-														// done button
-														SNew(SButton)
-															.ContentPadding(FMargin(6.0f, 2.0f))
-															.OnClicked(this, &SSessionLauncher::HandleDoneButtonClicked)
-															.ToolTipText(LOCTEXT("DoneButtonTooltip", "Close this page."))
+														SNew(SHorizontalBox)
+															.Visibility(this, &SSessionLauncher::HandleDoneButtonVisibility)
+
+														+ SHorizontalBox::Slot()
+															.AutoWidth()
+															.Padding(4.0f, 0.0f, 0.0f, 0.0f)
 															[
-																SNew(STextBlock)
-																	.Text(LOCTEXT("DoneButtonLabel", "Done"))
+																// re-run button
+																SNew(SButton)
+																	.ContentPadding(FMargin(6.0f, 2.0f))
+																	.OnClicked(this, &SSessionLauncher::HandleFinishButtonClicked)
+																	.ToolTipText(LOCTEXT("RerunButtonTooltip", "Execute the selected profile one more time."))
+																	[
+																		SNew(STextBlock)
+																			.Text(LOCTEXT("RerunButtonLabel", "Run Again"))
+																	]								
+															]
+
+														+ SHorizontalBox::Slot()
+															.AutoWidth()
+															.Padding(4.0f, 0.0f, 0.0f, 0.0f)
+															[
+																// done button
+																SNew(SButton)
+																	.ContentPadding(FMargin(6.0f, 2.0f))
+																	.OnClicked(this, &SSessionLauncher::HandleDoneButtonClicked)
+																	.ToolTipText(LOCTEXT("DoneButtonTooltip", "Close this page."))
+																	[
+																		SNew(STextBlock)
+																			.Text(LOCTEXT("DoneButtonLabel", "Done"))
+																	]
 															]
 													]
 											]
@@ -276,7 +274,6 @@ void SSessionLauncher::Construct( const FArguments& InArgs, const TSharedRef<SDo
 							]
 					]
 			]
-			]
 	];
 
 	ConstructUnderMajorTab->SetOnPersistVisualState(SDockTab::FOnPersistVisualState::CreateRaw(this, &SSessionLauncher::HandleMajorTabPersistVisualState));
diff --git a/Engine/Source/Developer/SessionLauncher/Private/Widgets/SSessionLauncher.h b/Engine/Source/Developer/SessionLauncher/Private/Widgets/SSessionLauncher.h
--- a/Engine/Source/Developer/SessionLauncher/Private/Widgets/SSessionLauncher.h
+++ b/Engine/Source/Developer/SessionLauncher/Private/Widgets/SSessionLauncher.h
@@ -1,9 +1,5 @@
 // Copyright 1998-2014 Epic Games, Inc. All Rights Reserved.
 
-/*=============================================================================
-	SSessionLauncher.h: Declares the SSessionLauncher class.
-=============================================================================*/
-
 #pragma once
 
 
@@ -34,10 +30,10 @@ public:
 	/**
 	 * Constructs the widget.
 	 *
-	 * @param InArgs - The Slate argument list.
-	 * @param ConstructUnderMajorTab - The major tab which will contain the session front-end.
-	 * @param ConstructUnderWindow - The window in which this widget is being constructed.
-	 * @param InModel - The view model to use.
+	 * @param InArgs The Slate argument list.
+	 * @param ConstructUnderMajorTab The major tab which will contain the session front-end.
+	 * @param ConstructUnderWindow The window in which this widget is being constructed.
+	 * @param InModel The view model to use.
 	 */
 	void Construct( const FArguments& InArgs, const TSharedRef<SDockTab>& ConstructUnderMajorTab, const TSharedPtr<SWindow>& ConstructUnderWindow, const FSessionLauncherModelRef& InModel );
 
@@ -48,8 +44,8 @@ protected:
 	/**
 	 * Fills the Window menu with menu items.
 	 *
-	 * @param MenuBuilder - The multi-box builder that should be filled with content for this pull-down menu.
-	 * @param RootMenuGroup - The root menu group.
+	 * @param MenuBuilder The multi-box builder that should be filled with content for this pull-down menu.
+	 * @param RootMenuGroup The root menu group.
 	 */
 	static void FillWindowMenu( FMenuBuilder& MenuBuilder, TSharedRef<FWorkspaceItem> RootMenuGroup );
 
diff --git a/Engine/Source/Developer/SessionLauncher/Private/Widgets/SSessionLauncherBuildTaskSettings.cpp b/Engine/Source/Developer/SessionLauncher/Private/Widgets/SSessionLauncherBuildTaskSettings.cpp
--- a/Engine/Source/Developer/SessionLauncher/Private/Widgets/SSessionLauncherBuildTaskSettings.cpp
+++ b/Engine/Source/Developer/SessionLauncher/Private/Widgets/SSessionLauncherBuildTaskSettings.cpp
@@ -1,9 +1,5 @@
 // Copyright 1998-2014 Epic Games, Inc. All Rights Reserved.
 
-/*=============================================================================
-	SSessionLauncherBuildTaskSettings.cpp: Implements the SSessionLauncherBuildTaskSettings class.
-=============================================================================*/
-
 #include "SessionLauncherPrivatePCH.h"
 
 
diff --git a/Engine/Source/Developer/SessionLauncher/Private/Widgets/SSessionLauncherBuildTaskSettings.h b/Engine/Source/Developer/SessionLauncher/Private/Widgets/SSessionLauncherBuildTaskSettings.h
--- a/Engine/Source/Developer/SessionLauncher/Private/Widgets/SSessionLauncherBuildTaskSettings.h
+++ b/Engine/Source/Developer/SessionLauncher/Private/Widgets/SSessionLauncherBuildTaskSettings.h
@@ -1,9 +1,5 @@
 // Copyright 1998-2014 Epic Games, Inc. All Rights Reserved.
 
-/*=============================================================================
-	SSessionLauncherBuildTaskSettings.h: Declares the SSessionLauncherBuildTaskSettings class.
-=============================================================================*/
-
 #pragma once
 
 
@@ -18,7 +14,6 @@ public:
 	SLATE_BEGIN_ARGS(SSessionLauncherBuildTaskSettings) { }
 	SLATE_END_ARGS()
 
-
 public:
 
 	/**
@@ -26,21 +21,16 @@ public:
 	 */
 	~SSessionLauncherBuildTaskSettings( );
 
-
 public:
 
 	/**
 	 * Constructs the widget.
 	 *
-	 * @param InArgs - The Slate argument list.
-	 * @param InModel - The data model.
+	 * @param InArgs The Slate argument list.
+	 * @param InModel The data model.
 	 */
 	void Construct( const FArguments& InArgs, const FSessionLauncherModelRef& InModel );
 
-
-private:
-
-
 private:
 
 	// Holds a pointer to the data model.
diff --git a/Engine/Source/Developer/SessionLauncher/Private/Widgets/SSessionLauncherDeployTaskSettings.cpp b/Engine/Source/Developer/SessionLauncher/Private/Widgets/SSessionLauncherDeployTaskSettings.cpp
--- a/Engine/Source/Developer/SessionLauncher/Private/Widgets/SSessionLauncherDeployTaskSettings.cpp
+++ b/Engine/Source/Developer/SessionLauncher/Private/Widgets/SSessionLauncherDeployTaskSettings.cpp
@@ -1,9 +1,5 @@
 // Copyright 1998-2014 Epic Games, Inc. All Rights Reserved.
 
-/*=============================================================================
-	SSessionLauncherDeployTaskSettings.cpp: Implements the SSessionLauncherDeployTaskSettings class.
-=============================================================================*/
-
 #include "SessionLauncherPrivatePCH.h"
 
 
diff --git a/Engine/Source/Developer/SessionLauncher/Private/Widgets/SSessionLauncherDeployTaskSettings.h b/Engine/Source/Developer/SessionLauncher/Private/Widgets/SSessionLauncherDeployTaskSettings.h
--- a/Engine/Source/Developer/SessionLauncher/Private/Widgets/SSessionLauncherDeployTaskSettings.h
+++ b/Engine/Source/Developer/SessionLauncher/Private/Widgets/SSessionLauncherDeployTaskSettings.h
@@ -1,9 +1,5 @@
 // Copyright 1998-2014 Epic Games, Inc. All Rights Reserved.
 
-/*=============================================================================
-	SSessionLauncherDeployTaskSettings.h: Declares the SSessionLauncherDeployTaskSettings class.
-=============================================================================*/
-
 #pragma once
 
 
@@ -18,7 +14,6 @@ public:
 	SLATE_BEGIN_ARGS(SSessionLauncherDeployTaskSettings) { }
 	SLATE_END_ARGS()
 
-
 public:
 
 	/**
@@ -26,21 +21,16 @@ public:
 	 */
 	~SSessionLauncherDeployTaskSettings( );
 
-
 public:
 
 	/**
 	 * Constructs the widget.
 	 *
-	 * @param InArgs - The Slate argument list.
-	 * @param InModel - The data model.
+	 * @param InArgs The Slate argument list.
+	 * @param InModel The data model.
 	 */
 	void Construct( const FArguments& InArgs, const FSessionLauncherModelRef& InModel );
 
-
-private:
-
-
 private:
 
 	// Holds a pointer to the data model.
diff --git a/Engine/Source/Developer/SessionLauncher/Private/Widgets/SSessionLauncherLaunchTaskSettings.cpp b/Engine/Source/Developer/SessionLauncher/Private/Widgets/SSessionLauncherLaunchTaskSettings.cpp
--- a/Engine/Source/Developer/SessionLauncher/Private/Widgets/SSessionLauncherLaunchTaskSettings.cpp
+++ b/Engine/Source/Developer/SessionLauncher/Private/Widgets/SSessionLauncherLaunchTaskSettings.cpp
@@ -1,9 +1,5 @@
 // Copyright 1998-2014 Epic Games, Inc. All Rights Reserved.
 
-/*=============================================================================
-	SSessionLauncherLaunchTaskSettings.cpp: Implements the SSessionLauncherLaunchTaskSettings class.
-=============================================================================*/
-
 #include "SessionLauncherPrivatePCH.h"
 
 
diff --git a/Engine/Source/Developer/SessionLauncher/Private/Widgets/SSessionLauncherLaunchTaskSettings.h b/Engine/Source/Developer/SessionLauncher/Private/Widgets/SSessionLauncherLaunchTaskSettings.h
--- a/Engine/Source/Developer/SessionLauncher/Private/Widgets/SSessionLauncherLaunchTaskSettings.h
+++ b/Engine/Source/Developer/SessionLauncher/Private/Widgets/SSessionLauncherLaunchTaskSettings.h
@@ -1,9 +1,5 @@
 // Copyright 1998-2014 Epic Games, Inc. All Rights Reserved.
 
-/*=============================================================================
-	SSessionLauncherLaunchTaskSettings.h: Declares the SSessionLauncherLaunchTaskSettings class.
-=============================================================================*/
-
 #pragma once
 
 
@@ -18,7 +14,6 @@ public:
 	SLATE_BEGIN_ARGS(SSessionLauncherLaunchTaskSettings) { }
 	SLATE_END_ARGS()
 
-
 public:
 
 	/**
@@ -26,21 +21,16 @@ public:
 	 */
 	~SSessionLauncherLaunchTaskSettings( );
 
-
 public:
 
 	/**
 	 * Constructs the widget.
 	 *
-	 * @param InArgs - The Slate argument list.
-	 * @param InModel - The data model.
+	 * @param InArgs The Slate argument list.
+	 * @param InModel The data model.
 	 */
 	void Construct( const FArguments& InArgs, const FSessionLauncherModelRef& InModel );
 
-
-private:
-
-
 private:
 
 	// Holds a pointer to the data model.
diff --git a/Engine/Source/Developer/SessionLauncher/Private/Widgets/SSessionLauncherMessageListRow.h b/Engine/Source/Developer/SessionLauncher/Private/Widgets/SSessionLauncherMessageListRow.h
--- a/Engine/Source/Developer/SessionLauncher/Private/Widgets/SSessionLauncherMessageListRow.h
+++ b/Engine/Source/Developer/SessionLauncher/Private/Widgets/SSessionLauncherMessageListRow.h
@@ -1,9 +1,5 @@
 // Copyright 1998-2014 Epic Games, Inc. All Rights Reserved.
 
-/*=============================================================================
-	SSessionLauncherTaskListRow.h: Declares the SSessionLauncherTaskListRow class.
-=============================================================================*/
-
 #pragma once
 
 
@@ -17,8 +13,7 @@ struct FLauncherMessage
 	FLauncherMessage(const FString& NewMessage, ELogVerbosity::Type InVerbosity)
 		: Message(NewMessage)
 		, Verbosity(InVerbosity)
-	{
-	}
+	{ }
 };
 
 
@@ -34,14 +29,13 @@ public:
 		SLATE_ARGUMENT(TSharedPtr<FLauncherMessage>, Message)
 	SLATE_END_ARGS()
 
-
 public:
 
 	/**
 	 * Constructs the widget.
 	 *
-	 * @param InArgs - The construction arguments.
-	 * @param InDeviceProxyManager - The device proxy manager to use.
+	 * @param InArgs The construction arguments.
+	 * @param InDeviceProxyManager The device proxy manager to use.
 	 */
 	void Construct( const FArguments& InArgs, const TSharedRef<STableViewBase>& InOwnerTableView )
 	{
@@ -50,14 +44,12 @@ public:
 		SMultiColumnTableRow<TSharedPtr<FLauncherMessage>>::Construct(FSuperRowType::FArguments(), InOwnerTableView);
 	}
 
-
 public:
 
 	/**
 	 * Generates the widget for the specified column.
 	 *
-	 * @param ColumnName - The name of the column to generate the widget for.
-	 *
+	 * @param ColumnName The name of the column to generate the widget for.
 	 * @return The widget.
 	 */
 	BEGIN_SLATE_FUNCTION_BUILD_OPTIMIZATION
@@ -79,7 +71,6 @@ public:
 	}
 	END_SLATE_FUNCTION_BUILD_OPTIMIZATION
 
-
 private:
 
 	// Callback for getting the task's status text.
@@ -100,8 +91,6 @@ private:
 		}
 	}
 
-
-
 private:
 
 	// Holds a pointer to the task that is displayed in this row.
diff --git a/Engine/Source/Developer/SessionLauncher/Private/Widgets/SSessionLauncherPreviewPage.cpp b/Engine/Source/Developer/SessionLauncher/Private/Widgets/SSessionLauncherPreviewPage.cpp
--- a/Engine/Source/Developer/SessionLauncher/Private/Widgets/SSessionLauncherPreviewPage.cpp
+++ b/Engine/Source/Developer/SessionLauncher/Private/Widgets/SSessionLauncherPreviewPage.cpp
@@ -1,9 +1,5 @@
 // Copyright 1998-2014 Epic Games, Inc. All Rights Reserved.
 
-/*=============================================================================
-	SSessionLauncherPreviewPage.cpp: Implements the SSessionLauncherPreviewPage class.
-=============================================================================*/
-
 #include "SessionLauncherPrivatePCH.h"
 
 
@@ -363,7 +359,6 @@ void SSessionLauncherPreviewPage::Construct( const FArguments& InArgs, const FSe
 									]
 							]
 					]
-
 /*
 				// package section
 				+ SGridPanel::Slot(0, 3)
@@ -389,7 +384,6 @@ void SSessionLauncherPreviewPage::Construct( const FArguments& InArgs, const FSe
 						SNullWidget::NullWidget
 					]
 */
-
 				// deploy section
 				+ SGridPanel::Slot(0, 5)
 					.ColumnSpan(3)
diff --git a/Engine/Source/Developer/SessionLauncher/Private/Widgets/SSessionLauncherPreviewPage.h b/Engine/Source/Developer/SessionLauncher/Private/Widgets/SSessionLauncherPreviewPage.h
--- a/Engine/Source/Developer/SessionLauncher/Private/Widgets/SSessionLauncherPreviewPage.h
+++ b/Engine/Source/Developer/SessionLauncher/Private/Widgets/SSessionLauncherPreviewPage.h
@@ -1,9 +1,5 @@
 // Copyright 1998-2014 Epic Games, Inc. All Rights Reserved.
 
-/*=============================================================================
-	SSessionLauncherPreviewPage.h: Declares the SSessionLauncherPreviewPage class.
-=============================================================================*/
-
 #pragma once
 
 
@@ -18,23 +14,20 @@ public:
 	SLATE_BEGIN_ARGS(SSessionLauncherPreviewPage) { }
 	SLATE_END_ARGS()
 
-
 public:
 
 	/**
 	 * Constructs the widget.
 	 *
-	 * @param InArgs - The Slate argument list.
-	 * @param InModel - The data model.
+	 * @param InArgs The Slate argument list.
+	 * @param InModel The data model.
 	 */
 	void Construct(	const FArguments& InArgs, const FSessionLauncherModelRef& InModel );
 
-
 public:
 
 	virtual void Tick( const FGeometry& AllottedGeometry, const double InCurrentTime, const float InDeltaTime ) OVERRIDE;
 
-
 protected:
 
 	/**
@@ -42,7 +35,6 @@ protected:
 	 */
 	void RefreshDeviceProxyList( );
 
-
 private:
 
 	// Callback for getting the text in the 'Build Configuration' text block.
@@ -102,13 +94,12 @@ private:
 	// Callback for determining the visibility of a validation error icon.
 	EVisibility HandleValidationErrorIconVisibility( ELauncherProfileValidationErrors::Type Error ) const;
 
-
 private:
 
 	// Holds the list of available device proxies.
 	TArray<ITargetDeviceProxyPtr> DeviceProxyList;
 
-	// Holds the device proxy list view .
+	// Holds the device proxy list view.
 	TSharedPtr<SListView<ITargetDeviceProxyPtr> > DeviceProxyListView;
 
 	// Holds a pointer to the data model.
diff --git a/Engine/Source/Developer/SessionLauncher/Private/Widgets/SSessionLauncherProfileSelector.h b/Engine/Source/Developer/SessionLauncher/Private/Widgets/SSessionLauncherProfileSelector.h
--- a/Engine/Source/Developer/SessionLauncher/Private/Widgets/SSessionLauncherProfileSelector.h
+++ b/Engine/Source/Developer/SessionLauncher/Private/Widgets/SSessionLauncherProfileSelector.h
@@ -1,9 +1,5 @@
 // Copyright 1998-2014 Epic Games, Inc. All Rights Reserved.
 
-/*=============================================================================
-	SSessionLauncherProfileSelector.h: Declares the SSessionLauncherProfileSelector class.
-=============================================================================*/
-
 #pragma once
 
 
@@ -21,7 +17,6 @@ public:
 	SLATE_BEGIN_ARGS(SSessionLauncherProfileSelector) { }
 	SLATE_END_ARGS()
 
-
 public:
 
 	/**
@@ -36,14 +31,13 @@ public:
 		}
 	}
 
-
 public:
 
 	/**
 	 * Constructs the widget.
 	 *
-	 * @param InArgs - The Slate argument list.
-	 * @param InModel - The data model.
+	 * @param InArgs The Slate argument list.
+	 * @param InModel The data model.
 	 */
 	void Construct( const FArguments& InArgs, const FSessionLauncherModelRef& InModel )
 	{
@@ -90,7 +84,6 @@ public:
 		Model->OnProfileListChanged().AddSP(this, &SSessionLauncherProfileSelector::HandleProfileManagerProfileListChanged);
 	}
 
-
 private:
 
 	// Callback for clicking the 'Add' button on the profile combo box.
@@ -195,7 +188,6 @@ private:
 		ProfileComboBox->SetSelectedItem(SelectedProfile);
 	}
 
-
 private:
 
 	// Holds the profile combo box.
diff --git a/Engine/Source/Developer/SessionLauncher/Private/Widgets/SSessionLauncherProgress.h b/Engine/Source/Developer/SessionLauncher/Private/Widgets/SSessionLauncherProgress.h
--- a/Engine/Source/Developer/SessionLauncher/Private/Widgets/SSessionLauncherProgress.h
+++ b/Engine/Source/Developer/SessionLauncher/Private/Widgets/SSessionLauncherProgress.h
@@ -1,9 +1,5 @@
 // Copyright 1998-2014 Epic Games, Inc. All Rights Reserved.
 
-/*=============================================================================
-	SSessionLauncherProgressPage.h: Declares the SSessionLauncherProgressPage class.
-=============================================================================*/
-
 #pragma once
 
 
@@ -21,7 +17,6 @@ public:
 	SLATE_BEGIN_ARGS(SSessionLauncherProgress) { }
 	SLATE_END_ARGS()
 
-
 public:
 
 	/**
@@ -32,13 +27,12 @@ public:
 
 	}
 
-
 public:
 
 	/**
 	 * Constructs the widget.
 	 *
-	 * @param InArgs - The Slate argument list.
+	 * @param InArgs The Slate argument list.
 	 */
 	void Construct( const FArguments& InArgs )
 	{
@@ -180,7 +174,7 @@ public:
 	/**
 	 * Sets the launcher worker to track the progress for.
 	 *
-	 * @param Worker - The launcher worker.
+	 * @param Worker The launcher worker.
 	 */
 	void SetLauncherWorker( const ILauncherWorkerRef& Worker )
 	{
@@ -211,6 +205,7 @@ public:
 		ClearButton->SetEnabled(MessageList.Num() > 0);
 		CopyButton->SetEnabled(MessageListView->GetNumItemsSelected() > 0);
 	}
+
 private:
 
 	void HandleOutputReceived(const FString& InMessage)
@@ -430,10 +425,10 @@ private:
 	TWeakPtr<ILauncherWorker> LauncherWorker;
 
 	// Holds the output log.
-	TArray<TSharedPtr<FString> > OutputList;
+	TArray<TSharedPtr<FString>> OutputList;
 
 	// Holds the output list view.
-	TSharedPtr<SListView<TSharedPtr<FString> > > OutputListView;
+	TSharedPtr<SListView<TSharedPtr<FString>>> OutputListView;
 
 	// Holds the progress bar.
 	TSharedPtr<SProgressBar> ProgressBar;
@@ -442,16 +437,16 @@ private:
 	TArray<ILauncherTaskPtr> TaskList;
 
 	// Holds the message list.
-	TArray< TSharedPtr<FLauncherMessage > > MessageList;
+	TArray< TSharedPtr<FLauncherMessage>> MessageList;
 
 	// Holds the pending message list.
-	TArray< TSharedPtr<FLauncherMessage > > PendingMessages;
+	TArray< TSharedPtr<FLauncherMessage>> PendingMessages;
 
 	// Holds the message list view.
-	TSharedPtr<SListView<TSharedPtr<FLauncherMessage>> > MessageListView;
+	TSharedPtr<SListView<TSharedPtr<FLauncherMessage>>> MessageListView;
 
 	// Holds the task list view.
-	TSharedPtr<SListView<ILauncherTaskPtr> > TaskListView;
+	TSharedPtr<SListView<ILauncherTaskPtr>> TaskListView;
 
 	// Holds the box of task statuses.
 	TSharedPtr<SVerticalBox> TaskStatusBox;
diff --git a/Engine/Source/Developer/SessionLauncher/Private/Widgets/SSessionLauncherSettings.cpp b/Engine/Source/Developer/SessionLauncher/Private/Widgets/SSessionLauncherSettings.cpp
--- a/Engine/Source/Developer/SessionLauncher/Private/Widgets/SSessionLauncherSettings.cpp
+++ b/Engine/Source/Developer/SessionLauncher/Private/Widgets/SSessionLauncherSettings.cpp
@@ -1,9 +1,5 @@
 // Copyright 1998-2014 Epic Games, Inc. All Rights Reserved.
 
-/*=============================================================================
-	SSessionLauncherSettings.cpp: Implements the SSessionLauncherSettings class.
-=============================================================================*/
-
 #include "SessionLauncherPrivatePCH.h"
 
 
diff --git a/Engine/Source/Developer/SessionLauncher/Private/Widgets/SSessionLauncherSettings.h b/Engine/Source/Developer/SessionLauncher/Private/Widgets/SSessionLauncherSettings.h
--- a/Engine/Source/Developer/SessionLauncher/Private/Widgets/SSessionLauncherSettings.h
+++ b/Engine/Source/Developer/SessionLauncher/Private/Widgets/SSessionLauncherSettings.h
@@ -14,34 +14,22 @@ namespace ESessionLauncherWizardPages
 	 */
 	enum Type
 	{
-		/**
-		 * The 'Build' page.
-		 */
+		/** The 'Build' page. */
 		BuildPage,
 
-		/**
-		 * The 'Cook' page.
-		 */
+		/** The 'Cook' page. */
 		CookPage,
 
-		/**
-		 * The 'Package' page.
-		 */
+		/** The 'Package' page. */
 		PackagePage,
 
-		/**
-		 * The 'Deploy' page.
-		 */
+		/** The 'Deploy' page. */
 		DeployPage,
 
-		/**
-		 * The 'Launch' page.
-		 */
+		/** The 'Launch' page. */
 		LaunchPage,
 
-		/**
-		 * The 'Preview' page.
-		 */
+		/** The 'Preview' page. */
 		PreviewPage
 	};
 }
@@ -58,7 +46,6 @@ public:
 	SLATE_BEGIN_ARGS(SSessionLauncherSettings) { }
 	SLATE_END_ARGS()
 
-
 public:
 
 	/**
@@ -66,19 +53,16 @@ public:
 	 */
 	~SSessionLauncherSettings( );
 
-
-
 public:
 
 	/**
 	 * Constructs the widget.
 	 *
-	 * @param InArgs - The Slate argument list.
-	 * @param InModel - The data model.
+	 * @param InArgs The Slate argument list.
+	 * @param InModel The data model.
 	 */
 	void Construct( const FArguments& InArgs, const FSessionLauncherModelRef& InModel );
 
-
 private:
 
 	// Callback for getting the visibility of the 'Select Profile' text block.
@@ -87,7 +71,6 @@ private:
 	// Callback for getting the visibility of the settings scroll box.
 	EVisibility HandleSettingsScrollBoxVisibility( ) const;
 
-
 private:
 
 	// Holds a pointer to the data model.
diff --git a/Engine/Source/Developer/SessionLauncher/Private/Widgets/SSessionLauncherTaskListRow.h b/Engine/Source/Developer/SessionLauncher/Private/Widgets/SSessionLauncherTaskListRow.h
--- a/Engine/Source/Developer/SessionLauncher/Private/Widgets/SSessionLauncherTaskListRow.h
+++ b/Engine/Source/Developer/SessionLauncher/Private/Widgets/SSessionLauncherTaskListRow.h
@@ -1,9 +1,5 @@
 // Copyright 1998-2014 Epic Games, Inc. All Rights Reserved.
 
-/*=============================================================================
-	SSessionLauncherTaskListRow.h: Declares the SSessionLauncherTaskListRow class.
-=============================================================================*/
-
 #pragma once
 
 
@@ -23,14 +19,13 @@ public:
 		SLATE_ARGUMENT(ILauncherTaskPtr, Task)
 	SLATE_END_ARGS()
 
-
 public:
 
 	/**
 	 * Constructs the widget.
 	 *
-	 * @param InArgs - The construction arguments.
-	 * @param InDeviceProxyManager - The device proxy manager to use.
+	 * @param InArgs The construction arguments.
+	 * @param InDeviceProxyManager The device proxy manager to use.
 	 */
 	void Construct( const FArguments& InArgs )
 	{
@@ -39,14 +34,12 @@ public:
 		SMultiColumnTableRow<ILauncherTaskPtr>::Construct(FSuperRowType::FArguments(), InArgs._OwnerTableView.ToSharedRef());
 	}
 
-
 public:
 
 	/**
 	 * Generates the widget for the specified column.
 	 *
-	 * @param ColumnName - The name of the column to generate the widget for.
-	 *
+	 * @param ColumnName The name of the column to generate the widget for.
 	 * @return The widget.
 	 */
 	BEGIN_SLATE_FUNCTION_BUILD_OPTIMIZATION
@@ -115,7 +108,6 @@ public:
 	}
 	END_SLATE_FUNCTION_BUILD_OPTIMIZATION
 
-
 private:
 
 	// Callback for getting the duration of the task.
@@ -244,7 +236,6 @@ private:
 		return EVisibility::Hidden;
 	}
 
-
 private:
 
 	// Holds a pointer to the task that is displayed in this row.
diff --git a/Engine/Source/Developer/SessionLauncher/Private/Widgets/SSessionLauncherValidation.h b/Engine/Source/Developer/SessionLauncher/Private/Widgets/SSessionLauncherValidation.h
--- a/Engine/Source/Developer/SessionLauncher/Private/Widgets/SSessionLauncherValidation.h
+++ b/Engine/Source/Developer/SessionLauncher/Private/Widgets/SSessionLauncherValidation.h
@@ -1,9 +1,5 @@
 // Copyright 1998-2014 Epic Games, Inc. All Rights Reserved.
 
-/*=============================================================================
-	SSessionLauncherValidation.h: Declares the SSessionLauncherValidation class.
-=============================================================================*/
-
 #pragma once
 
 
@@ -27,8 +23,8 @@ public:
 	/**
 	 * Constructs the widget.
 	 *
-	 * @param InArgs - The Slate argument list.
-	 * @param InModel - The data model.
+	 * @param InArgs The Slate argument list.
+	 * @param InModel The data model.
 	 */
 	void Construct( const FArguments& InArgs, const FSessionLauncherModelRef& InModel )
 	{
@@ -112,15 +108,14 @@ public:
 		];
 	}
 
-
 protected:
 
 	/**
 	 * Creates a widget for a validation message.
 	 *
-	 * @param IconName - The name of the message icon.
-	 * @param MessageText - The message text.
-	 * @param MessageType - The message type.
+	 * @param IconName The name of the message icon.
+	 * @param MessageText The message text.
+	 * @param MessageType The message type.
 	 */
 	TSharedRef<SWidget> MakeValidationMessage( const TCHAR* IconName, const FString& MessageText, ELauncherProfileValidationErrors::Type Message )
 	{
@@ -144,7 +139,6 @@ protected:
 			];
 	}
 
-
 private:
 
 	// Callback for getting the visibility state of a validation message.
@@ -163,7 +157,6 @@ private:
 		return EVisibility::Collapsed;
 	}
 
-
 private:
 
 	// Holds a pointer to the data model.
diff --git a/Engine/Source/Developer/SessionLauncher/Private/Widgets/Shared/SSessionLauncherBuildConfigurationSelector.h b/Engine/Source/Developer/SessionLauncher/Private/Widgets/Shared/SSessionLauncherBuildConfigurationSelector.h
--- a/Engine/Source/Developer/SessionLauncher/Private/Widgets/Shared/SSessionLauncherBuildConfigurationSelector.h
+++ b/Engine/Source/Developer/SessionLauncher/Private/Widgets/Shared/SSessionLauncherBuildConfigurationSelector.h
@@ -1,9 +1,5 @@
 // Copyright 1998-2014 Epic Games, Inc. All Rights Reserved.
 
-/*=============================================================================
-	SSessionLauncherBuildConfigurationSelector.h: Declares the SSessionLauncherBuildConfigurationSelector class.
-=============================================================================*/
-
 #pragma once
 
 
@@ -31,14 +27,13 @@ public:
 		SLATE_ATTRIBUTE(FString, Text)
 	SLATE_END_ARGS()
 
-
 public:
 
 	/**
 	 * Constructs the widget.
 	 *
-	 * @param InArgs - The Slate argument list.
-	 * @param InModel - The data model.
+	 * @param InArgs The Slate argument list.
+	 * @param InModel The data model.
 	 */
 	void Construct(	const FArguments& InArgs )
 	{
@@ -80,7 +75,6 @@ public:
 		];
 	}
 
-
 private:
 
 	// Callback for clicking a menu entry.
@@ -89,7 +83,6 @@ private:
 		OnConfigurationSelected.ExecuteIfBound(Configuration);
 	}
 
-
 private:
 
 	// Holds a delegate to be invoked when a build configuration has been selected.
diff --git a/Engine/Source/Developer/SessionLauncher/Private/Widgets/Shared/SSessionLauncherFormLabel.h b/Engine/Source/Developer/SessionLauncher/Private/Widgets/Shared/SSessionLauncherFormLabel.h
--- a/Engine/Source/Developer/SessionLauncher/Private/Widgets/Shared/SSessionLauncherFormLabel.h
+++ b/Engine/Source/Developer/SessionLauncher/Private/Widgets/Shared/SSessionLauncherFormLabel.h
@@ -1,9 +1,5 @@
 // Copyright 1998-2014 Epic Games, Inc. All Rights Reserved.
 
-/*=============================================================================
-	SSessionLauncherFormLabel.h: Declares the SSessionLauncherFormLabel class.
-=============================================================================*/
-
 #pragma once
 
 
@@ -21,30 +17,23 @@ public:
 		, _LabelText()
 	{ }
 
-		/**
-		 * The tool tip text for the validation error icon.
-		 */
+		/** The tool tip text for the validation error icon. */
 		SLATE_ATTRIBUTE(FText, ErrorToolTipText)
 
-		/**
-		 * The visibility of the validation error icon.
-		 */
+		/** The visibility of the validation error icon. */
 		SLATE_ATTRIBUTE(EVisibility, ErrorVisibility)
 
-		/**
-		 * The text of the form label.
-		 */
+		/** The text of the form label. */
 		SLATE_ATTRIBUTE(FText, LabelText)
 
 	SLATE_END_ARGS()
 
-
 public:
 
 	/**
 	 * Constructs the widget.
 	 *
-	 * @param InArgs - The Slate argument list.
+	 * @param InArgs The Slate argument list.
 	 */
 	void Construct(	const FArguments& InArgs )
 	{
diff --git a/Engine/Source/Developer/SessionLauncher/Private/Widgets/Toolbar/SSessionLauncherToolbar.cpp b/Engine/Source/Developer/SessionLauncher/Private/Widgets/Toolbar/SSessionLauncherToolbar.cpp
--- a/Engine/Source/Developer/SessionLauncher/Private/Widgets/Toolbar/SSessionLauncherToolbar.cpp
+++ b/Engine/Source/Developer/SessionLauncher/Private/Widgets/Toolbar/SSessionLauncherToolbar.cpp
@@ -1,9 +1,5 @@
 // Copyright 1998-2014 Epic Games, Inc. All Rights Reserved.
 
-/*=============================================================================
-	SSessionLauncherToolbar.cpp: Implements the SSessionLauncherToolbar class.
-=============================================================================*/
-
 #include "SessionLauncherPrivatePCH.h"
 
 
diff --git a/Engine/Source/Developer/SessionLauncher/Private/Widgets/Toolbar/SSessionLauncherToolbar.h b/Engine/Source/Developer/SessionLauncher/Private/Widgets/Toolbar/SSessionLauncherToolbar.h
--- a/Engine/Source/Developer/SessionLauncher/Private/Widgets/Toolbar/SSessionLauncherToolbar.h
+++ b/Engine/Source/Developer/SessionLauncher/Private/Widgets/Toolbar/SSessionLauncherToolbar.h
@@ -1,15 +1,10 @@
 // Copyright 1998-2014 Epic Games, Inc. All Rights Reserved.
 
-/*=============================================================================
-	SSessionLauncherToolbar.h: Declares the SSessionLauncherToolbar class.
-=============================================================================*/
-
 #pragma once
 
 
 #define LOCTEXT_NAMESPACE "SSessionLauncherTaskSelector"
 
-
 class SSessionLauncher;
 
 
@@ -38,15 +33,15 @@ public:
 	/**
 	 * Constructs the widget.
 	 *
-	 * @param InArgs - The Slate argument list.
-	 * @param InModel - The data model.
+	 * @param InArgs The Slate argument list.
+	 * @param InModel The data model.
 	 */
 	void Construct( const FArguments& InArgs, const FSessionLauncherModelRef& InModel, SSessionLauncher* InLauncher );
 
 	/**
 	 * Sets the currently selected task.
 	 *
-	 * @param InTask - The identifier of the task to set.
+	 * @param InTask The identifier of the task to set.
 	 */
 	void SetCurrentTask( ESessionLauncherTasks::Type InTask )
 	{
@@ -63,8 +58,7 @@ protected:
 	/**
 	 * Creates the toolbar widget.
 	 *
-	 * @param CommandList - The command list to use for the toolbar buttons.
-	 *
+	 * @param CommandList The command list to use for the toolbar buttons.
 	 * @return The toolbar widget.
 	 */
 	TSharedRef<SWidget> MakeToolbar( const TSharedRef<FUICommandList>& CommandList );
diff --git a/Engine/Source/Developer/SessionLauncher/Public/Interfaces/ISessionLauncherModule.h b/Engine/Source/Developer/SessionLauncher/Public/Interfaces/ISessionLauncherModule.h
--- a/Engine/Source/Developer/SessionLauncher/Public/Interfaces/ISessionLauncherModule.h
+++ b/Engine/Source/Developer/SessionLauncher/Public/Interfaces/ISessionLauncherModule.h
@@ -1,9 +1,5 @@
 // Copyright 1998-2014 Epic Games, Inc. All Rights Reserved.
 
-/*=============================================================================
-	ISessionLauncherModule.h: Declares the ISessionLauncherModule interface.
-=============================================================================*/
-
 #pragma once
 
 
@@ -18,8 +14,7 @@ public:
 	/**
 	 * Creates a session launcher progress panel widget.
 	 *
-	 * @param LauncherWorker - The launcher worker.
-	 *
+	 * @param LauncherWorker The launcher worker.
 	 * @return The new widget.
 	 */
 	virtual TSharedRef<class SWidget> CreateLauncherProgressPanel( const ILauncherWorkerRef& LauncherWorker ) = 0;
diff --git a/Engine/Source/Developer/SessionLauncher/Public/SessionLauncher.h b/Engine/Source/Developer/SessionLauncher/Public/SessionLauncher.h
--- a/Engine/Source/Developer/SessionLauncher/Public/SessionLauncher.h
+++ b/Engine/Source/Developer/SessionLauncher/Public/SessionLauncher.h
@@ -1,9 +1,5 @@
 // Copyright 1998-2014 Epic Games, Inc. All Rights Reserved.
 
-/*=============================================================================
-	SessionLauncher.h: SessionLauncher module public header file.
-=============================================================================*/
-
 #pragma once
 
 
diff --git a/Engine/Source/Developer/SessionMessages/Classes/SessionServiceMessages.h b/Engine/Source/Developer/SessionMessages/Classes/SessionServiceMessages.h
--- a/Engine/Source/Developer/SessionMessages/Classes/SessionServiceMessages.h
+++ b/Engine/Source/Developer/SessionMessages/Classes/SessionServiceMessages.h
@@ -1,9 +1,5 @@
 // Copyright 1998-2014 Epic Games, Inc. All Rights Reserved.
 
-/*=============================================================================
-	SessionServiceMessages.h: Declares message types used by session services.
-=============================================================================*/
-
 #pragma once
 
 #include "SessionServiceMessages.generated.h"
diff --git a/Engine/Source/Developer/SessionMessages/Private/SessionMessagesModule.cpp b/Engine/Source/Developer/SessionMessages/Private/SessionMessagesModule.cpp
--- a/Engine/Source/Developer/SessionMessages/Private/SessionMessagesModule.cpp
+++ b/Engine/Source/Developer/SessionMessages/Private/SessionMessagesModule.cpp
@@ -1,13 +1,8 @@
 // Copyright 1998-2014 Epic Games, Inc. All Rights Reserved.
 
-/*=============================================================================
-	SessionMessagesModule.cpp: Implements the FSessionMessagesModule class.
-=============================================================================*/
-
 #include "SessionMessagesPrivatePCH.h"
 
 
-
 /**
  * Implements the SessionMessages module.
  */
@@ -16,18 +11,15 @@ class FSessionMessagesModule
 {
 public:
 
-	// Begin IModuleInterface interface
+	// IModuleInterface interface
 
-	virtual void StartupModule( ) OVERRIDE { }
+	virtual void StartupModule( ) override { }
+	virtual void ShutdownModule( ) override { }
 
-	virtual void ShutdownModule( ) OVERRIDE { }
-
-	virtual bool SupportsDynamicReloading( ) OVERRIDE
+	virtual bool SupportsDynamicReloading( ) override
 	{
 		return true;
 	}
-
-	// End IModuleInterface interface
 };
 
 
diff --git a/Engine/Source/Developer/SessionMessages/Private/SessionMessagesPrivatePCH.h b/Engine/Source/Developer/SessionMessages/Private/SessionMessagesPrivatePCH.h
--- a/Engine/Source/Developer/SessionMessages/Private/SessionMessagesPrivatePCH.h
+++ b/Engine/Source/Developer/SessionMessages/Private/SessionMessagesPrivatePCH.h
@@ -1,24 +1,11 @@
 // Copyright 1998-2014 Epic Games, Inc. All Rights Reserved.
 
-/*=============================================================================
-	SessionMessagesPrivatePCH.h: Pre-compiled header file for the SessionMessages module.
-=============================================================================*/
+#pragma once
 
-#ifndef SESSIONMESSAGES_PRIVATEPCH_H
-#define SESSIONMESSAGES_PRIVATEPCH_H
-
-
-#include "../Public/SessionMessages.h"
+#include "SessionMessages.h"
 
 
 /* Dependencies
  *****************************************************************************/
 
 #include "ModuleManager.h"
-
-
-/* Private includes
- *****************************************************************************/
-
-
-#endif
diff --git a/Engine/Source/Developer/SessionMessages/Public/SessionMessages.h b/Engine/Source/Developer/SessionMessages/Public/SessionMessages.h
--- a/Engine/Source/Developer/SessionMessages/Public/SessionMessages.h
+++ b/Engine/Source/Developer/SessionMessages/Public/SessionMessages.h
@@ -1,9 +1,5 @@
 // Copyright 1998-2014 Epic Games, Inc. All Rights Reserved.
 
-/*=============================================================================
-	SessionMessages.h: SessionMessages module public header file.
-=============================================================================*/
-
 #pragma once
 
 
diff --git a/Engine/Source/Developer/SessionServices/Private/SessionInfo.h b/Engine/Source/Developer/SessionServices/Private/SessionInfo.h
--- a/Engine/Source/Developer/SessionServices/Private/SessionInfo.h
+++ b/Engine/Source/Developer/SessionServices/Private/SessionInfo.h
@@ -1,9 +1,5 @@
 // Copyright 1998-2014 Epic Games, Inc. All Rights Reserved.
 
-/*=============================================================================
-	SessionInfo.h: Declares the FSessionInfo class.
-=============================================================================*/
-
 #pragma once
 
 
@@ -108,9 +104,9 @@ public:
 
 public:	
 
-	// Begin SessionInfo interface
+	// SessionInfo interface
 
-	virtual void GetInstances( TArray<ISessionInstanceInfoPtr>& OutInstances ) const OVERRIDE
+	virtual void GetInstances( TArray<ISessionInstanceInfoPtr>& OutInstances ) const override
 	{
 		OutInstances.Empty();
 
@@ -120,47 +116,47 @@ public:
 		}
 	}
 
-	virtual FDateTime GetLastUpdateTime( ) OVERRIDE
+	virtual FDateTime GetLastUpdateTime( ) override
 	{
 		return LastUpdateTime;
 	}
 
-	virtual const int32 GetNumInstances( ) const OVERRIDE
+	virtual const int32 GetNumInstances( ) const override
 	{
 		return Instances.Num();
 	}
 
-	virtual const FGuid& GetSessionId( ) const OVERRIDE
+	virtual const FGuid& GetSessionId( ) const override
 	{
 		return SessionId;
 	}
 
-	virtual const FString& GetSessionName( ) const OVERRIDE
+	virtual const FString& GetSessionName( ) const override
 	{
 		return SessionName;
 	}
 
-	virtual const FString& GetSessionOwner( ) const OVERRIDE
+	virtual const FString& GetSessionOwner( ) const override
 	{
 		return SessionOwner;
 	}
 
-	virtual const bool IsStandalone( ) const OVERRIDE
+	virtual const bool IsStandalone( ) const override
 	{
 		return Standalone;
 	}
 
-	virtual FOnSessionInstanceDiscovered& OnInstanceDiscovered( ) OVERRIDE
+	virtual FOnSessionInstanceDiscovered& OnInstanceDiscovered( ) override
 	{
 		return InstanceDiscoveredDelegate;
 	}
 
-	virtual FOnSessionLogReceived& OnLogReceived( ) OVERRIDE
+	virtual FOnSessionLogReceived& OnLogReceived( ) override
 	{
 		return LogReceivedDelegate;
 	}
 
-	virtual void Terminate( ) OVERRIDE
+	virtual void Terminate( ) override
 	{
 		for (TMap<FMessageAddress, TSharedPtr<FSessionInstanceInfo> >::TIterator It(Instances); It; ++It)
 		{
@@ -168,8 +164,6 @@ public:
 		}
 	}
 
-	// End SessionInfo interface
-
 private:
 
 	// Handles received log messages.
diff --git a/Engine/Source/Developer/SessionServices/Private/SessionInstanceInfo.h b/Engine/Source/Developer/SessionServices/Private/SessionInstanceInfo.h
--- a/Engine/Source/Developer/SessionServices/Private/SessionInstanceInfo.h
+++ b/Engine/Source/Developer/SessionServices/Private/SessionInstanceInfo.h
@@ -1,9 +1,5 @@
 // Copyright 1998-2014 Epic Games, Inc. All Rights Reserved.
 
-/*=============================================================================
-	GameInstanceInfo.h: Declares the FSessionInstanceInfo class.
-=============================================================================*/
-
 #pragma once
 
 
@@ -90,7 +86,7 @@ public:
 
 public:	
 
-	// Begin IGameInstanceInfo interface
+	// IGameInstanceInfo interface
 
 	virtual void ExecuteCommand( const FString& CommandString )
 	{
@@ -100,82 +96,82 @@ public:
 		}
 	}
 
-	virtual const FString GetBuildDate() const OVERRIDE
+	virtual const FString GetBuildDate() const override
 	{
 		return BuildDate;
 	}
 
-	virtual const FString GetCurrentLevel() const OVERRIDE
+	virtual const FString GetCurrentLevel() const override
 	{
 		return CurrentLevel;
 	}
 
-	virtual const FString GetDeviceName() const OVERRIDE
+	virtual const FString GetDeviceName() const override
 	{
 		return DeviceName;
 	}
 
-	virtual int32 GetEngineVersion( ) const OVERRIDE
+	virtual int32 GetEngineVersion( ) const override
 	{
 		return EngineVersion;
 	}
 
-	virtual const FGuid GetInstanceId( ) const OVERRIDE
+	virtual const FGuid GetInstanceId( ) const override
 	{
 		return InstanceId;
 	}
 
-	virtual const FString GetInstanceName() const OVERRIDE
+	virtual const FString GetInstanceName() const override
 	{
 		return InstanceName;
 	}
 
-	virtual const FString GetInstanceType() const OVERRIDE
+	virtual const FString GetInstanceType() const override
 	{
 		return InstanceType;
 	}
 
-	virtual FDateTime GetLastUpdateTime( ) OVERRIDE
+	virtual FDateTime GetLastUpdateTime( ) override
 	{
 		return LastUpdateTime;
 	}
 
-	virtual const TArray<FSessionLogMessagePtr>& GetLog( ) OVERRIDE
+	virtual const TArray<FSessionLogMessagePtr>& GetLog( ) override
 	{
 		return LogMessages;
 	}
 
-	virtual ISessionInfoPtr GetOwnerSession( ) OVERRIDE
+	virtual ISessionInfoPtr GetOwnerSession( ) override
 	{
 		return Owner.Pin();
 	}
 
-	virtual const FString& GetPlatformName( ) const OVERRIDE
+	virtual const FString& GetPlatformName( ) const override
 	{
 		return PlatformName;
 	}
 
-	virtual float GetWorldTimeSeconds( ) const OVERRIDE
+	virtual float GetWorldTimeSeconds( ) const override
 	{
 		return WorldTimeSeconds;
 	}
 
-	virtual const bool IsConsole() const OVERRIDE
+	virtual const bool IsConsole() const override
 	{
 		return IsConsoleBuild;
 	}
 
-	virtual FOnSessionInstanceLogReceived& OnLogReceived( ) OVERRIDE
+	virtual FOnSessionInstanceLogReceived& OnLogReceived( ) override
 	{
 		return LogReceivedDelegate;
 	}
 
-	virtual bool PlayHasBegun( ) const OVERRIDE
+	virtual bool PlayHasBegun( ) const override
 	{
 		return HasBegunPlay;
 	}
 
-	virtual void Terminate( ) OVERRIDE
+	virtual void Terminate( ) override
 	{
 		if (MessageEndpoint.IsValid() && EngineAddress.IsValid())
 		{
@@ -183,8 +179,6 @@ public:
 		}
 	}
 
-	// End IGameInstanceInfo interface
-
 private:
 
 	// Handles FSessionServiceLog messages.
diff --git a/Engine/Source/Developer/SessionServices/Private/SessionManager.cpp b/Engine/Source/Developer/SessionServices/Private/SessionManager.cpp
--- a/Engine/Source/Developer/SessionServices/Private/SessionManager.cpp
+++ b/Engine/Source/Developer/SessionServices/Private/SessionManager.cpp
@@ -1,9 +1,5 @@
 // Copyright 1998-2014 Epic Games, Inc. All Rights Reserved.
 
-/*=============================================================================
-	SessionManager.cpp: Implements the FSessionManager class.
-=============================================================================*/
-
 #include "SessionServicesPrivatePCH.h"
 
 
diff --git a/Engine/Source/Developer/SessionServices/Private/SessionManager.h b/Engine/Source/Developer/SessionServices/Private/SessionManager.h
--- a/Engine/Source/Developer/SessionServices/Private/SessionManager.h
+++ b/Engine/Source/Developer/SessionServices/Private/SessionManager.h
@@ -1,9 +1,5 @@
 // Copyright 1998-2014 Epic Games, Inc. All Rights Reserved.
 
-/*=============================================================================
-	SessionManager.h: Declares the FSessionManager interface.
-=============================================================================*/
-
 #pragma once
 
 
@@ -32,59 +28,54 @@ public:
 
 	// Begin ISessionManager Interface
 
-	virtual void AddOwner( const FString& InOwner ) OVERRIDE;
-
-	virtual void GetSelectedInstances( TArray<ISessionInstanceInfoPtr>& OutInstances) const OVERRIDE;
+	virtual void AddOwner( const FString& InOwner ) override;
+	virtual void GetSelectedInstances( TArray<ISessionInstanceInfoPtr>& OutInstances) const override;
 
-	virtual const ISessionInfoPtr& GetSelectedSession( ) const OVERRIDE
+	virtual const ISessionInfoPtr& GetSelectedSession( ) const override
 	{
 		return SelectedSession;
 	}
 
-	virtual void GetSessions( TArray<ISessionInfoPtr>& OutSessions ) const OVERRIDE;
+	virtual void GetSessions( TArray<ISessionInfoPtr>& OutSessions ) const override;
 
-	virtual bool IsInstanceSelected( const ISessionInstanceInfoRef& Instance ) const OVERRIDE
+	virtual bool IsInstanceSelected( const ISessionInstanceInfoRef& Instance ) const override
 	{
 		return ((Instance->GetOwnerSession() == SelectedSession) && !DeselectedInstances.Contains(Instance));
 	}
 
-	virtual FOnCanSelectSession& OnCanSelectSession( ) OVERRIDE
+	virtual FOnCanSelectSession& OnCanSelectSession( ) override
 	{
 		return CanSelectSessionDelegate;
 	}
 
-	virtual FOnSessionInstanceSelectionChanged& OnInstanceSelectionChanged( ) OVERRIDE
+	virtual FOnSessionInstanceSelectionChanged& OnInstanceSelectionChanged( ) override
 	{
 		return InstanceSelectionChangedDelegate;
 	}
 
-	virtual FOnSessionLogReceived& OnLogReceived( ) OVERRIDE
+	virtual FOnSessionLogReceived& OnLogReceived( ) override
 	{
 		return LogReceivedDelegate;
 	}
 
-	virtual FOnSelectedSessionChanged& OnSelectedSessionChanged( ) OVERRIDE
+	virtual FOnSelectedSessionChanged& OnSelectedSessionChanged( ) override
 	{
 		return SelectedSessionChangedDelegate;
 	}
 
-	virtual FSimpleMulticastDelegate& OnSessionsUpdated( ) OVERRIDE
+	virtual FSimpleMulticastDelegate& OnSessionsUpdated( ) override
 	{
 		return SessionsUpdatedDelegate;
 	}
 
-	virtual FSimpleMulticastDelegate& OnSessionInstanceUpdated( ) OVERRIDE
+	virtual FSimpleMulticastDelegate& OnSessionInstanceUpdated( ) override
 	{
 		return SessionInstanceUpdatedDelegate;
 	}
 
-	virtual void RemoveOwner( const FString& InOwner ) OVERRIDE;
-
-	virtual bool SelectSession( const ISessionInfoPtr& Session ) OVERRIDE;
-
-	virtual bool SetInstanceSelected( const ISessionInstanceInfoPtr& Instance, bool Selected ) OVERRIDE;
-
-	// End ISessionManager Interface
+	virtual void RemoveOwner( const FString& InOwner ) override;
+	virtual bool SelectSession( const ISessionInfoPtr& Session ) override;
+	virtual bool SetInstanceSelected( const ISessionInstanceInfoPtr& Instance, bool Selected ) override;
 
 protected:
 
diff --git a/Engine/Source/Developer/SessionServices/Private/SessionService.cpp b/Engine/Source/Developer/SessionServices/Private/SessionService.cpp
--- a/Engine/Source/Developer/SessionServices/Private/SessionService.cpp
+++ b/Engine/Source/Developer/SessionServices/Private/SessionService.cpp
@@ -1,9 +1,5 @@
 // Copyright 1998-2014 Epic Games, Inc. All Rights Reserved.
 
-/*=============================================================================
-	SessionService.cpp: Implements the FSessionService class.
-=============================================================================*/
-
 #include "SessionServicesPrivatePCH.h"
 
 
diff --git a/Engine/Source/Developer/SessionServices/Private/SessionService.h b/Engine/Source/Developer/SessionServices/Private/SessionService.h
--- a/Engine/Source/Developer/SessionServices/Private/SessionService.h
+++ b/Engine/Source/Developer/SessionServices/Private/SessionService.h
@@ -1,9 +1,5 @@
 // Copyright 1998-2014 Epic Games, Inc. All Rights Reserved.
 
-/*=============================================================================
-	SessionService.h: Declares the FSessionService class.
-=============================================================================*/
-
 #pragma once
 
 
@@ -30,29 +26,24 @@ public:
 
 public:
 
-	// Begin FOutputDevice interface
+	// FOutputDevice interface
 
-	virtual void Serialize( const TCHAR* Data, ELogVerbosity::Type Verbosity, const class FName& Category ) OVERRIDE
+	virtual void Serialize( const TCHAR* Data, ELogVerbosity::Type Verbosity, const class FName& Category ) override
 	{
 		SendLog(Data, Verbosity, Category);
 	}
 
-	// End FOutputDevice interface
-
 public:
 
-	// Begin ISessionService interface
+	// ISessionService interface
 
-	virtual bool IsRunning( ) OVERRIDE
+	virtual bool IsRunning( ) override
 	{
 		return MessageEndpoint.IsValid();
 	}
 
-	virtual bool Start( ) OVERRIDE;
-
-	virtual void Stop( ) OVERRIDE;
-
-	// End ISessionService interface
+	virtual bool Start( ) override;
+	virtual void Stop( ) override;
 
 protected:
 
diff --git a/Engine/Source/Developer/SessionServices/Private/SessionServicesModule.cpp b/Engine/Source/Developer/SessionServices/Private/SessionServicesModule.cpp
--- a/Engine/Source/Developer/SessionServices/Private/SessionServicesModule.cpp
+++ b/Engine/Source/Developer/SessionServices/Private/SessionServicesModule.cpp
@@ -1,11 +1,6 @@
 // Copyright 1998-2014 Epic Games, Inc. All Rights Reserved.
 
-/*=============================================================================
-	SessionServicesModule.cpp: Implements the FSessionServicesModule class.
-=============================================================================*/
-
 #include "SessionServicesPrivatePCH.h"
-
 #include "ModuleManager.h"
 
 
@@ -29,9 +24,9 @@ public:
 
 public:
 
-	// Begin ISessionServicesModule interface
+	// ISessionServicesModule interface
 
-	virtual ISessionManagerRef GetSessionManager( ) OVERRIDE
+	virtual ISessionManagerRef GetSessionManager( ) override
 	{
 		if (!SessionManager.IsValid())
 		{
@@ -41,7 +36,7 @@ public:
 		return SessionManager.ToSharedRef();
 	}
 
-	virtual ISessionServiceRef GetSessionService( ) OVERRIDE
+	virtual ISessionServiceRef GetSessionService( ) override
 	{
 		if (!SessionService.IsValid())
 		{
@@ -51,26 +46,22 @@ public:
 		return SessionService.ToSharedRef();
 	}
 
-	// End ISessionServicesModule interface
-
 public:
 
-	// Begin IModuleInterface interface
+	// IModuleInterface interface
 
-	virtual void StartupModule( ) OVERRIDE
+	virtual void StartupModule( ) override
 	{
 		MessageBusPtr = IMessagingModule::Get().GetDefaultBus();
 		check(MessageBusPtr.IsValid());
 	}
 
-	virtual void ShutdownModule( ) OVERRIDE
+	virtual void ShutdownModule( ) override
 	{
 		SessionManager.Reset();
 		SessionService.Reset();
 	}
 
-	// End IModuleInterface interface
-
 private:
 	
 	// Holds a weak pointer to the message bus.
diff --git a/Engine/Source/Developer/SessionServices/Private/SessionServicesPrivatePCH.h b/Engine/Source/Developer/SessionServices/Private/SessionServicesPrivatePCH.h
--- a/Engine/Source/Developer/SessionServices/Private/SessionServicesPrivatePCH.h
+++ b/Engine/Source/Developer/SessionServices/Private/SessionServicesPrivatePCH.h
@@ -1,9 +1,5 @@
 // Copyright 1998-2014 Epic Games, Inc. All Rights Reserved.
 
-/*=============================================================================
-	SessionServicesPrivatePCH.h: Pre-compiled header file for the SessionServices module.
-=============================================================================*/
-
 #pragma once
 
 #include "SessionServices.h"
@@ -24,4 +20,4 @@
 #include "SessionInstanceInfo.h"
 #include "SessionInfo.h"
 #include "SessionManager.h"
-#include "SessionService.h"
\ No newline at end of file
+#include "SessionService.h"
diff --git a/Engine/Source/Developer/SessionServices/Public/Interfaces/ISessionInfo.h b/Engine/Source/Developer/SessionServices/Public/Interfaces/ISessionInfo.h
--- a/Engine/Source/Developer/SessionServices/Public/Interfaces/ISessionInfo.h
+++ b/Engine/Source/Developer/SessionServices/Public/Interfaces/ISessionInfo.h
@@ -1,9 +1,5 @@
 // Copyright 1998-2014 Epic Games, Inc. All Rights Reserved.
 
-/*=============================================================================
-	ISessionInfo.h: Declares the ISessionInfo interface.
-=============================================================================*/
-
 #pragma once
 
 
diff --git a/Engine/Source/Developer/SessionServices/Public/Interfaces/ISessionInstanceInfo.h b/Engine/Source/Developer/SessionServices/Public/Interfaces/ISessionInstanceInfo.h
--- a/Engine/Source/Developer/SessionServices/Public/Interfaces/ISessionInstanceInfo.h
+++ b/Engine/Source/Developer/SessionServices/Public/Interfaces/ISessionInstanceInfo.h
@@ -1,9 +1,5 @@
 // Copyright 1998-2014 Epic Games, Inc. All Rights Reserved.
 
-/*=============================================================================
-	IGameInstanceInfo.h: Declares the IGameInstanceInfo interface.
-=============================================================================*/
-
 #pragma once
 
 
diff --git a/Engine/Source/Developer/SessionServices/Public/Interfaces/ISessionManager.h b/Engine/Source/Developer/SessionServices/Public/Interfaces/ISessionManager.h
--- a/Engine/Source/Developer/SessionServices/Public/Interfaces/ISessionManager.h
+++ b/Engine/Source/Developer/SessionServices/Public/Interfaces/ISessionManager.h
@@ -1,9 +1,5 @@
 // Copyright 1998-2014 Epic Games, Inc. All Rights Reserved.
 
-/*=============================================================================
-	ISessionManager.h: Declares the ISessionManager interface.
-=============================================================================*/
-
 #pragma once
 
 
diff --git a/Engine/Source/Developer/SessionServices/Public/Interfaces/ISessionService.h b/Engine/Source/Developer/SessionServices/Public/Interfaces/ISessionService.h
--- a/Engine/Source/Developer/SessionServices/Public/Interfaces/ISessionService.h
+++ b/Engine/Source/Developer/SessionServices/Public/Interfaces/ISessionService.h
@@ -1,9 +1,5 @@
 // Copyright 1998-2014 Epic Games, Inc. All Rights Reserved.
 
-/*=============================================================================
-	ISessionService.h: Declares the ISessionService interface.
-=============================================================================*/
-
 #pragma once
 
 
@@ -48,4 +44,4 @@ public:
 	 * Virtual destructor.
 	 */
 	virtual ~ISessionService( ) { }
-};
\ No newline at end of file
+};
diff --git a/Engine/Source/Developer/SessionServices/Public/Interfaces/ISessionServicesModule.h b/Engine/Source/Developer/SessionServices/Public/Interfaces/ISessionServicesModule.h
--- a/Engine/Source/Developer/SessionServices/Public/Interfaces/ISessionServicesModule.h
+++ b/Engine/Source/Developer/SessionServices/Public/Interfaces/ISessionServicesModule.h
@@ -1,9 +1,5 @@
 // Copyright 1998-2014 Epic Games, Inc. All Rights Reserved.
 
-/*=============================================================================
-	ISessionServicesModule.h: Declares the ISessionServicesModule interface.
-=============================================================================*/
-
 #pragma once
 
 
diff --git a/Engine/Source/Developer/SessionServices/Public/Interfaces/SessionLogMessage.h b/Engine/Source/Developer/SessionServices/Public/Interfaces/SessionLogMessage.h
--- a/Engine/Source/Developer/SessionServices/Public/Interfaces/SessionLogMessage.h
+++ b/Engine/Source/Developer/SessionServices/Public/Interfaces/SessionLogMessage.h
@@ -1,9 +1,5 @@
 // Copyright 1998-2014 Epic Games, Inc. All Rights Reserved.
 
-/*=============================================================================
-	LauncherLogMessage.h: Declares the FLauncherLogMessage structure.
-=============================================================================*/
-
 #pragma once
 
 
diff --git a/Engine/Source/Developer/SessionServices/Public/SessionServices.h b/Engine/Source/Developer/SessionServices/Public/SessionServices.h
--- a/Engine/Source/Developer/SessionServices/Public/SessionServices.h
+++ b/Engine/Source/Developer/SessionServices/Public/SessionServices.h
@@ -1,9 +1,5 @@
 // Copyright 1998-2014 Epic Games, Inc. All Rights Reserved.
 
-/*=============================================================================
-	SessionServices.h: SessionServices module public header file.
-=============================================================================*/
-
 #pragma once
 
 
diff --git a/Engine/Source/Developer/Settings/Private/SettingsCategory.h b/Engine/Source/Developer/Settings/Private/SettingsCategory.h
--- a/Engine/Source/Developer/Settings/Private/SettingsCategory.h
+++ b/Engine/Source/Developer/Settings/Private/SettingsCategory.h
@@ -1,9 +1,5 @@
 // Copyright 1998-2014 Epic Games, Inc. All Rights Reserved.
 
-/*=============================================================================
-	SettingsCategory.h: Declares the FSettingsCategory class.
-=============================================================================*/
-
 #pragma once
 
 
@@ -102,34 +98,34 @@ public:
 
 public:
 
-	// Begin ISettingsCategory interface
+	// ISettingsCategory interface
 
-	virtual const FText& GetDescription( ) const OVERRIDE
+	virtual const FText& GetDescription( ) const override
 	{
 		return Description;
 	}
 
-	virtual const FText& GetDisplayName( ) const OVERRIDE
+	virtual const FText& GetDisplayName( ) const override
 	{
 		return DisplayName;
 	}
 
-	virtual const FName& GetIconName( ) const OVERRIDE
+	virtual const FName& GetIconName( ) const override
 	{
 		return IconName;
 	}
 
-	virtual const FName& GetName( ) const OVERRIDE
+	virtual const FName& GetName( ) const override
 	{
 		return Name;
 	}
 
-	virtual ISettingsSectionPtr GetSection( const FName& SectionName ) const OVERRIDE
+	virtual ISettingsSectionPtr GetSection( const FName& SectionName ) const override
 	{
 		return Sections.FindRef(SectionName);
 	}
 
-	virtual int32 GetSections( TArray<ISettingsSectionPtr>& OutSections ) const OVERRIDE
+	virtual int32 GetSections( TArray<ISettingsSectionPtr>& OutSections ) const override
 	{
 		OutSections.Empty(Sections.Num());
 
@@ -141,8 +137,6 @@ public:
 		return OutSections.Num();
 	}
 
-	// End ISettingsCategory interface
-
 private:
 
 	// Holds the category's description text.
diff --git a/Engine/Source/Developer/Settings/Private/SettingsContainer.h b/Engine/Source/Developer/Settings/Private/SettingsContainer.h
--- a/Engine/Source/Developer/Settings/Private/SettingsContainer.h
+++ b/Engine/Source/Developer/Settings/Private/SettingsContainer.h
@@ -1,9 +1,5 @@
 // Copyright 1998-2014 Epic Games, Inc. All Rights Reserved.
 
-/*=============================================================================
-	SettingsContainer.h: Declares the FSettingsContainer class.
-=============================================================================*/
-
 #pragma once
 
 
@@ -18,7 +14,7 @@ public:
 	/**
 	 * Creates and initializes a new instance.
 	 *
-	 * @param InName - The container's name.
+	 * @param InName The container's name.
 	 */
 	FSettingsContainer( const FName& InName )
 		: Name(InName)
@@ -107,16 +103,16 @@ public:
 
 public:
 
-	// Begin ISettingsContainer interface
+	// ISettingsContainer interface
 
-	virtual void Describe( const FText& InDisplayName, const FText& InDescription, const FName& InIconName ) OVERRIDE
+	virtual void Describe( const FText& InDisplayName, const FText& InDescription, const FName& InIconName ) override
 	{
 		Description = InDescription;
 		DisplayName = InDisplayName;
 		IconName = InIconName;
 	}
 
-	virtual void DescribeCategory( const FName& CategoryName, const FText& InDisplayName, const FText& InDescription, const FName& InIconName ) OVERRIDE
+	virtual void DescribeCategory( const FName& CategoryName, const FText& InDisplayName, const FText& InDescription, const FName& InIconName ) override
 	{
 		TSharedPtr<FSettingsCategory>& Category = Categories.FindOrAdd(CategoryName);
 
@@ -129,7 +125,7 @@ public:
 		CategoryModifiedDelegate.Broadcast(CategoryName);
 	}
 
-	virtual int32 GetCategories( TArray<ISettingsCategoryPtr>& OutCategories ) const OVERRIDE
+	virtual int32 GetCategories( TArray<ISettingsCategoryPtr>& OutCategories ) const override
 	{
 		OutCategories.Empty(Categories.Num());
 
@@ -141,43 +137,41 @@ public:
 		return OutCategories.Num();
 	}
 
-	virtual ISettingsCategoryPtr GetCategory( const FName& CategoryName ) const OVERRIDE
+	virtual ISettingsCategoryPtr GetCategory( const FName& CategoryName ) const override
 	{
 		return Categories.FindRef(CategoryName);
 	}
 
-	virtual const FText& GetDescription( ) const OVERRIDE
+	virtual const FText& GetDescription( ) const override
 	{
 		return Description;
 	}
 
-	virtual const FText& GetDisplayName( ) const OVERRIDE
+	virtual const FText& GetDisplayName( ) const override
 	{
 		return DisplayName;
 	}
 
-	virtual const FName& GetIconName( ) const OVERRIDE
+	virtual const FName& GetIconName( ) const override
 	{
 		return IconName;
 	}
 
-	virtual const FName& GetName( ) const OVERRIDE
+	virtual const FName& GetName( ) const override
 	{
 		return Name;
 	}
 
-	virtual FOnSettingsContainerCategoryModified& OnCategoryModified( ) OVERRIDE
+	virtual FOnSettingsContainerCategoryModified& OnCategoryModified( ) override
 	{
 		return CategoryModifiedDelegate;
 	}
 
-	virtual FOnSettingsContainerSectionRemoved& OnSectionRemoved( ) OVERRIDE
+	virtual FOnSettingsContainerSectionRemoved& OnSectionRemoved( ) override
 	{
 		return SectionRemovedDelegate;
 	}
 
-	// End ISettingsContainer interface
-
 private:
 
 	// Holds the collection of setting categories
diff --git a/Engine/Source/Developer/Settings/Private/SettingsModule.cpp b/Engine/Source/Developer/Settings/Private/SettingsModule.cpp
--- a/Engine/Source/Developer/Settings/Private/SettingsModule.cpp
+++ b/Engine/Source/Developer/Settings/Private/SettingsModule.cpp
@@ -1,9 +1,5 @@
 // Copyright 1998-2014 Epic Games, Inc. All Rights Reserved.
 
-/*=============================================================================
-	SettingsModule.cpp: Implements the FSettingsModule class.
-=============================================================================*/
-
 #include "SettingsPrivatePCH.h"
 
 
@@ -18,24 +14,24 @@ class FSettingsModule
 {
 public:
 
-	// Begin ISettingsModule interface
+	// ISettingsModule interface
 
-	virtual void GetContainerNames( TArray<FName>& OutNames ) const OVERRIDE
+	virtual void GetContainerNames( TArray<FName>& OutNames ) const override
 	{
 		ContainerNamesToContainers.GenerateKeyArray(OutNames);
 	}
 
-	virtual ISettingsContainerPtr GetContainer( const FName& ContainerName ) OVERRIDE
+	virtual ISettingsContainerPtr GetContainer( const FName& ContainerName ) override
 	{
 		return ContainerNamesToContainers.FindRef(ContainerName);
 	}
 
-	virtual ISettingsSectionPtr RegisterSettings( const FName& ContainerName, const FName& CategoryName, const FName& SectionName, const FText& DisplayName, const FText& Description, const TWeakObjectPtr<UObject>& SettingsObject, const FSettingsSectionDelegates& Delegates ) OVERRIDE
+	virtual ISettingsSectionPtr RegisterSettings( const FName& ContainerName, const FName& CategoryName, const FName& SectionName, const FText& DisplayName, const FText& Description, const TWeakObjectPtr<UObject>& SettingsObject, const FSettingsSectionDelegates& Delegates ) override
 	{
 		return FindOrAddContainer(ContainerName)->AddSection(CategoryName, SectionName, DisplayName, Description, SettingsObject, Delegates);
 	}
 
-	virtual ISettingsSectionPtr RegisterSettings( const FName& ContainerName, const FName& CategoryName, const FName& SectionName, const FText& DisplayName, const FText& Description, const TSharedRef<SWidget>& CustomWidget, const FSettingsSectionDelegates& Delegates ) OVERRIDE
+	virtual ISettingsSectionPtr RegisterSettings( const FName& ContainerName, const FName& CategoryName, const FName& SectionName, const FText& DisplayName, const FText& Description, const TSharedRef<SWidget>& CustomWidget, const FSettingsSectionDelegates& Delegates ) override
 	{
 		return FindOrAddContainer(ContainerName)->AddSection(CategoryName, SectionName, DisplayName, Description, CustomWidget, Delegates);
 	}
@@ -60,7 +56,7 @@ public:
 		ContainerNamesToViewers.Remove(ContainerName);
 	}
 
-	virtual void UnregisterSettings( const FName& ContainerName, const FName& CategoryName, const FName& SectionName ) OVERRIDE
+	virtual void UnregisterSettings( const FName& ContainerName, const FName& CategoryName, const FName& SectionName ) override
 	{
 		TSharedPtr<FSettingsContainer> Container = ContainerNamesToContainers.FindRef(ContainerName);
 
@@ -70,13 +66,11 @@ public:
 		}
 	}
 
-	// End ISettingsModule interface
-
 public:
 
-	// Begin IModuleInterface interface
+	// IModuleInterface interface
 	
-	virtual void StartupModule() OVERRIDE
+	virtual void StartupModule( ) override
 	{
 		// @todo gmp: move this into the modules that own these setting categories
 		TSharedRef<FSettingsContainer> EditorSettingsContainer = FindOrAddContainer("Editor");
@@ -95,17 +89,14 @@ public:
 		ProjectSettingsContainer->DescribeCategory("Plugins", LOCTEXT("ProjectPluginsCategoryName", "Plugins"), LOCTEXT("ProjectPluginsCategoryDescription", "Plugins settings description text here"), "SettingsEditor.Category_Plugins");
 	}
 
-	virtual void ShutdownModule() OVERRIDE { }
-	
-	// End IModuleInterface interface
+	virtual void ShutdownModule( ) override { }
 
 protected:
 
 	/**
 	 * Finds or adds the specified settings container.
 	 *
-	 * @param ContainerName - The name of the container to find or add.
-	 *
+	 * @param ContainerName The name of the container to find or add.
 	 * @return The container.
 	 */
 	TSharedRef<FSettingsContainer> FindOrAddContainer( const FName& ContainerName )
diff --git a/Engine/Source/Developer/Settings/Private/SettingsPrivatePCH.h b/Engine/Source/Developer/Settings/Private/SettingsPrivatePCH.h
--- a/Engine/Source/Developer/Settings/Private/SettingsPrivatePCH.h
+++ b/Engine/Source/Developer/Settings/Private/SettingsPrivatePCH.h
@@ -1,13 +1,9 @@
 // Copyright 1998-2014 Epic Games, Inc. All Rights Reserved.
 
-/*=============================================================================
-	SettingsPrivatePCH.h: Pre-compiled header file for the Settings module.
-=============================================================================*/
-
 #pragma once
 
 
-#include "../Public/Settings.h"
+#include "Settings.h"
 
 
 /* Dependencies
diff --git a/Engine/Source/Developer/Settings/Private/SettingsSection.h b/Engine/Source/Developer/Settings/Private/SettingsSection.h
--- a/Engine/Source/Developer/Settings/Private/SettingsSection.h
+++ b/Engine/Source/Developer/Settings/Private/SettingsSection.h
@@ -1,9 +1,5 @@
 // Copyright 1998-2014 Epic Games, Inc. All Rights Reserved.
 
-/*=============================================================================
-	SettingsSection.h: Declares the FSettingsSection class.
-=============================================================================*/
-
 #pragma once
 
 
@@ -69,9 +65,9 @@ public:
 
 public:
 
-	// Begin ISettingsSection interface
+	// ISettingsSection interface
 
-	virtual bool CanEdit( ) const OVERRIDE
+	virtual bool CanEdit( ) const override
 	{
 		if (CanEditDelegate.IsBound())
 		{
@@ -81,32 +77,32 @@ public:
 		return true;
 	}
 
-	virtual bool CanExport( ) const OVERRIDE
+	virtual bool CanExport( ) const override
 	{
 		return (ExportDelegate.IsBound() || (SettingsObject.IsValid() && SettingsObject->GetClass()->HasAnyClassFlags(CLASS_Config)));
 	}
 
-	virtual bool CanImport( ) const OVERRIDE
+	virtual bool CanImport( ) const override
 	{
 		return (ImportDelegate.IsBound() || (SettingsObject.IsValid() && SettingsObject->GetClass()->HasAnyClassFlags(CLASS_Config)));
 	}
 
-	virtual bool CanResetDefaults( ) const OVERRIDE
+	virtual bool CanResetDefaults( ) const override
 	{
 		return (ResetDefaultsDelegate.IsBound() || (SettingsObject.IsValid() && SettingsObject->GetClass()->HasAnyClassFlags(CLASS_Config) && !SettingsObject->GetClass()->HasAnyClassFlags(CLASS_DefaultConfig)));
 	}
 
-	virtual bool CanSave( ) const OVERRIDE
+	virtual bool CanSave( ) const override
 	{
 		return (SaveDelegate.IsBound() || (SettingsObject.IsValid() && SettingsObject->GetClass()->HasAnyClassFlags(CLASS_Config)));
 	}
 
-	virtual bool CanSaveDefaults( ) const OVERRIDE
+	virtual bool CanSaveDefaults( ) const override
 	{
 		return (SaveDefaultsDelegate.IsBound() || (SettingsObject.IsValid() && SettingsObject->GetClass()->HasAnyClassFlags(CLASS_Config) && !SettingsObject->GetClass()->HasAnyClassFlags(CLASS_DefaultConfig)));
 	}
 
-	virtual bool Export( const FString& Filename ) OVERRIDE
+	virtual bool Export( const FString& Filename ) override
 	{
 		if (ExportDelegate.IsBound())
 		{
@@ -123,37 +119,37 @@ public:
 		return false;
 	}
 
-	virtual ISettingsCategoryWeakPtr GetCategory( ) OVERRIDE
+	virtual ISettingsCategoryWeakPtr GetCategory( ) override
 	{
 		return Category;
 	}
 
-	virtual TWeakPtr<SWidget> GetCustomWidget( ) const OVERRIDE
+	virtual TWeakPtr<SWidget> GetCustomWidget( ) const override
 	{
 		return CustomWidget;
 	}
 
-	virtual const FText& GetDescription( ) const OVERRIDE
+	virtual const FText& GetDescription( ) const override
 	{
 		return Description;
 	}
 
-	virtual const FText& GetDisplayName( ) const OVERRIDE
+	virtual const FText& GetDisplayName( ) const override
 	{
 		return DisplayName;
 	}
 
-	virtual const FName& GetName( ) const OVERRIDE
+	virtual const FName& GetName( ) const override
 	{
 		return Name;
 	}
 
-	virtual TWeakObjectPtr<UObject> GetSettingsObject( ) const OVERRIDE
+	virtual TWeakObjectPtr<UObject> GetSettingsObject( ) const override
 	{
 		return SettingsObject;
 	}
 
-	virtual FText GetStatus( ) const OVERRIDE
+	virtual FText GetStatus( ) const override
 	{
 		if (StatusDelegate.IsBound())
 		{
@@ -163,7 +159,7 @@ public:
 		return FText::GetEmpty();
 	}
 
-	virtual bool HasDefaultSettingsObject( ) OVERRIDE
+	virtual bool HasDefaultSettingsObject( ) override
 	{
 		if (!SettingsObject.IsValid())
 		{
@@ -173,7 +169,7 @@ public:
 		return SettingsObject->GetClass()->HasAnyClassFlags(CLASS_DefaultConfig);
 	}
 
-	virtual bool Import( const FString& Filename ) OVERRIDE
+	virtual bool Import( const FString& Filename ) override
 	{
 		if (ImportDelegate.IsBound())
 		{
@@ -190,7 +186,7 @@ public:
 		return false;	
 	}
 
-	virtual bool ResetDefaults( ) OVERRIDE
+	virtual bool ResetDefaults( ) override
 	{
 		if (ResetDefaultsDelegate.IsBound())
 		{
@@ -214,7 +210,7 @@ public:
 		return false;
 	}
 
-	virtual bool Save( ) OVERRIDE
+	virtual bool Save( ) override
 	{
 		if (ModifiedDelegate.IsBound() && !ModifiedDelegate.Execute())
 		{
@@ -243,7 +239,7 @@ public:
 		return false;
 	}
 
-	virtual bool SaveDefaults( ) OVERRIDE
+	virtual bool SaveDefaults( ) override
 	{
 		if (SaveDefaultsDelegate.IsBound())
 		{
@@ -261,8 +257,6 @@ public:
 		return false;
 	}
 
-	// End ISettingsSection interface
-
 private:
 
 	// Holds a pointer to the settings category that owns this section.
diff --git a/Engine/Source/Developer/Settings/Public/Interfaces/ISettingsCategory.h b/Engine/Source/Developer/Settings/Public/Interfaces/ISettingsCategory.h
--- a/Engine/Source/Developer/Settings/Public/Interfaces/ISettingsCategory.h
+++ b/Engine/Source/Developer/Settings/Public/Interfaces/ISettingsCategory.h
@@ -1,9 +1,5 @@
 // Copyright 1998-2014 Epic Games, Inc. All Rights Reserved.
 
-/*=============================================================================
-	ISettingsCategory.h: Declares the ISettingsCategory interface.
-=============================================================================*/
-
 #pragma once
 
 
diff --git a/Engine/Source/Developer/Settings/Public/Interfaces/ISettingsContainer.h b/Engine/Source/Developer/Settings/Public/Interfaces/ISettingsContainer.h
--- a/Engine/Source/Developer/Settings/Public/Interfaces/ISettingsContainer.h
+++ b/Engine/Source/Developer/Settings/Public/Interfaces/ISettingsContainer.h
@@ -1,9 +1,5 @@
 // Copyright 1998-2014 Epic Games, Inc. All Rights Reserved.
 
-/*=============================================================================
-	ISettingsContainer.h: Declares the ISettingsContainer interface.
-=============================================================================*/
-
 #pragma once
 
 
diff --git a/Engine/Source/Developer/Settings/Public/Interfaces/ISettingsModule.h b/Engine/Source/Developer/Settings/Public/Interfaces/ISettingsModule.h
--- a/Engine/Source/Developer/Settings/Public/Interfaces/ISettingsModule.h
+++ b/Engine/Source/Developer/Settings/Public/Interfaces/ISettingsModule.h
@@ -1,9 +1,5 @@
 // Copyright 1998-2014 Epic Games, Inc. All Rights Reserved.
 
-/*=============================================================================
-	ISettingsModule.h: Declares the ISettingsModule interface.
-=============================================================================*/
-
 #pragma once
 
 
diff --git a/Engine/Source/Developer/Settings/Public/Interfaces/ISettingsSection.h b/Engine/Source/Developer/Settings/Public/Interfaces/ISettingsSection.h
--- a/Engine/Source/Developer/Settings/Public/Interfaces/ISettingsSection.h
+++ b/Engine/Source/Developer/Settings/Public/Interfaces/ISettingsSection.h
@@ -1,9 +1,5 @@
 // Copyright 1998-2014 Epic Games, Inc. All Rights Reserved.
 
-/*=============================================================================
-	ISettingsSection.h: Declares the ISettingsSection interface.
-=============================================================================*/
-
 #pragma once
 
 
@@ -71,44 +67,28 @@ DECLARE_DELEGATE_RetVal(FText, FOnSettingsSectionStatus)
  */
 struct FSettingsSectionDelegates
 {
-	/**
-	 * Holds a delegate that is executed to check whether a settings section can be edited (default = true).
-	 */
+	/** A delegate that is executed to check whether a settings section can be edited (default = true). */
 	FOnSettingsSectionCanEdit CanEditDelegate;
 
-	/**
-	 * Holds a delegate that is executed when a settings section should export its values to a file (default = false).
-	 */
+	/** A delegate that is executed when a settings section should export its values to a file (default = false). */
 	FOnSettingsSectionExport ExportDelegate;
 
-	/**
-	 * Holds a delegate that is executed when a settings section should import its values from a file (default = false).
-	 */
+	/** A delegate that is executed when a settings section should import its values from a file (default = false). */
 	FOnSettingsSectionImport ImportDelegate;
 
-	/**
-	 * Holds a delegate that is executed when a settings section has been modified (default = true).
-	 */
+	/** A delegate that is executed when a settings section has been modified (default = true). */
 	FOnSettingsSectionModified ModifiedDelegate;
 
-	/**
-	 * Holds a delegate that is executed when a settings section should have its values reset to default (default = false).
-	 */
+	/** A delegate that is executed when a settings section should have its values reset to default (default = false). */
 	FOnSettingsSectionResetDefaults ResetDefaultsDelegate;
 
-	/**
-	 * Holds a delegate that is executed when a settings section should have its values saved as default.
-	 */
+	/** A delegate that is executed when a settings section should have its values saved as default. */
 	FOnSettingsSectionSaveDefaults SaveDefaultsDelegate;
 
-	/**
-	 * Holds a delegate that is executed when a settings section should have its values saved.
-	 */
+	/** A delegate that is executed when a settings section should have its values saved. */
 	FOnSettingsSectionSave SaveDelegate;
 
-	/**
-	 * Holds a delegate that is executed to retrieve a status message for a settings section.
-	 */
+	/** A delegate that is executed to retrieve a status message for a settings section. */
 	FOnSettingsSectionStatus StatusDelegate;
 };
 
diff --git a/Engine/Source/Developer/Settings/Public/Interfaces/ISettingsViewer.h b/Engine/Source/Developer/Settings/Public/Interfaces/ISettingsViewer.h
--- a/Engine/Source/Developer/Settings/Public/Interfaces/ISettingsViewer.h
+++ b/Engine/Source/Developer/Settings/Public/Interfaces/ISettingsViewer.h
@@ -1,9 +1,5 @@
 // Copyright 1998-2014 Epic Games, Inc. All Rights Reserved.
 
-/*=============================================================================
-	ISettingsViewer.h: Declares the ISettingsViewer interface.
-=============================================================================*/
-
 #pragma once
 
 
diff --git a/Engine/Source/Developer/Settings/Public/Settings.h b/Engine/Source/Developer/Settings/Public/Settings.h
--- a/Engine/Source/Developer/Settings/Public/Settings.h
+++ b/Engine/Source/Developer/Settings/Public/Settings.h
@@ -1,9 +1,5 @@
 // Copyright 1998-2014 Epic Games, Inc. All Rights Reserved.
 
-/*=============================================================================
-	Settings.h: Settings module public header file.
-=============================================================================*/
-
 #pragma once
 
 
diff --git a/Engine/Source/Developer/SettingsEditor/Private/Models/SettingsEditorModel.h b/Engine/Source/Developer/SettingsEditor/Private/Models/SettingsEditorModel.h
--- a/Engine/Source/Developer/SettingsEditor/Private/Models/SettingsEditorModel.h
+++ b/Engine/Source/Developer/SettingsEditor/Private/Models/SettingsEditorModel.h
@@ -1,9 +1,5 @@
 // Copyright 1998-2014 Epic Games, Inc. All Rights Reserved.
 
-/*=============================================================================
-	SettingsEditorModel.h: Declares the FSettingsEditorModel interface.
-=============================================================================*/
-
 #pragma once
 
 
@@ -36,24 +32,24 @@ public:
 
 public:
 
-	// Begin ISettingsEditorModel interface
+	// ISettingsEditorModel interface
 
-	virtual const ISettingsSectionPtr& GetSelectedSection( ) const OVERRIDE
+	virtual const ISettingsSectionPtr& GetSelectedSection( ) const override
 	{
 		return SelectedSection;
 	}
 
-	virtual const ISettingsContainerRef& GetSettingsContainer( ) const OVERRIDE
+	virtual const ISettingsContainerRef& GetSettingsContainer( ) const override
 	{
 		return SettingsContainer;
 	}
 
-	virtual FSimpleMulticastDelegate& OnSelectionChanged( ) OVERRIDE
+	virtual FSimpleMulticastDelegate& OnSelectionChanged( ) override
 	{
 		return OnSelectionChangedDelegate;
 	}
 
-	virtual void SelectSection( const ISettingsSectionPtr& Section ) OVERRIDE
+	virtual void SelectSection( const ISettingsSectionPtr& Section ) override
 	{
 		if (Section == SelectedSection)
 		{
@@ -65,8 +61,6 @@ public:
 		OnSelectionChangedDelegate.Broadcast();
 	}
 
-	// End ISettingsEditorModel interface
-
 private:
 
 	// Handles the removal of sections from the settings container.
diff --git a/Engine/Source/Developer/SettingsEditor/Private/SettingsEditorModule.cpp b/Engine/Source/Developer/SettingsEditor/Private/SettingsEditorModule.cpp
--- a/Engine/Source/Developer/SettingsEditor/Private/SettingsEditorModule.cpp
+++ b/Engine/Source/Developer/SettingsEditor/Private/SettingsEditorModule.cpp
@@ -1,9 +1,5 @@
 // Copyright 1998-2014 Epic Games, Inc. All Rights Reserved.
 
-/*=============================================================================
-	SettingsEditorModule.cpp: Implements the FSettingsEditorModule class.
-=============================================================================*/
-
 #include "SettingsEditorPrivatePCH.h"
 
 
@@ -15,19 +11,17 @@ class FSettingsEditorModule
 {
 public:
 
-	// Begin ISettingsEditorModule interface
+	// ISettingsEditorModule interface
 
-	virtual TSharedRef<SWidget> CreateEditor( const ISettingsEditorModelRef& Model ) OVERRIDE
+	virtual TSharedRef<SWidget> CreateEditor( const ISettingsEditorModelRef& Model ) override
 	{
 		return SNew(SSettingsEditor, Model);
 	}
 
-	virtual ISettingsEditorModelRef CreateModel( const ISettingsContainerRef& SettingsContainer ) OVERRIDE
+	virtual ISettingsEditorModelRef CreateModel( const ISettingsContainerRef& SettingsContainer ) override
 	{
 		return MakeShareable(new FSettingsEditorModel(SettingsContainer));
 	}
-
-	// End ISettingsEditorModule interface
 };
 
 
diff --git a/Engine/Source/Developer/SettingsEditor/Private/SettingsEditorPrivatePCH.h b/Engine/Source/Developer/SettingsEditor/Private/SettingsEditorPrivatePCH.h
--- a/Engine/Source/Developer/SettingsEditor/Private/SettingsEditorPrivatePCH.h
+++ b/Engine/Source/Developer/SettingsEditor/Private/SettingsEditorPrivatePCH.h
@@ -1,13 +1,9 @@
 // Copyright 1998-2014 Epic Games, Inc. All Rights Reserved.
 
-/*=============================================================================
-	SettingsEditorPrivatePCH.h: Pre-compiled header file for the SettingsEditor module.
-=============================================================================*/
-
 #pragma once
 
 
-#include "../Public/SettingsEditor.h"
+#include "SettingsEditor.h"
 
 
 /* Dependencies
diff --git a/Engine/Source/Developer/SettingsEditor/Private/Widgets/SSettingsEditor.cpp b/Engine/Source/Developer/SettingsEditor/Private/Widgets/SSettingsEditor.cpp
--- a/Engine/Source/Developer/SettingsEditor/Private/Widgets/SSettingsEditor.cpp
+++ b/Engine/Source/Developer/SettingsEditor/Private/Widgets/SSettingsEditor.cpp
@@ -1,9 +1,5 @@
 // Copyright 1998-2014 Epic Games, Inc. All Rights Reserved.
 
-/*=============================================================================
-	SSettingsEditor.h: Implements the SSettingsEditor class.
-=============================================================================*/
-
 #include "SettingsEditorPrivatePCH.h"
 
 
diff --git a/Engine/Source/Developer/SettingsEditor/Private/Widgets/SSettingsEditor.h b/Engine/Source/Developer/SettingsEditor/Private/Widgets/SSettingsEditor.h
--- a/Engine/Source/Developer/SettingsEditor/Private/Widgets/SSettingsEditor.h
+++ b/Engine/Source/Developer/SettingsEditor/Private/Widgets/SSettingsEditor.h
@@ -1,9 +1,5 @@
 // Copyright 1998-2014 Epic Games, Inc. All Rights Reserved.
 
-/*=============================================================================
-	SSettingsEditor.h: Declares the SSettingsEditor class.
-=============================================================================*/
-
 #pragma once
 
 
@@ -38,19 +34,15 @@ public:
 
 public:
 
-	// Begin SCompoundWidget interface
-
-	virtual void Tick( const FGeometry& AllottedGeometry, const double InCurrentTime, const float InDeltaTime ) OVERRIDE;
+	// SCompoundWidget interface
 
-	// End SCompoundWidget interface
+	virtual void Tick( const FGeometry& AllottedGeometry, const double InCurrentTime, const float InDeltaTime ) override;
 
 public:
 
-	// Begin FNotifyHook interface
-
-	virtual void NotifyPostChange( const FPropertyChangedEvent& PropertyChangedEvent, class FEditPropertyChain* PropertyThatChanged ) OVERRIDE;
+	// FNotifyHook interface
 
-	// End FNotifyHook interface
+	virtual void NotifyPostChange( const FPropertyChangedEvent& PropertyChangedEvent, class FEditPropertyChain* PropertyThatChanged ) override;
 
 protected:
 
diff --git a/Engine/Source/Developer/SettingsEditor/Private/Widgets/SSettingsEditorCheckoutNotice.cpp b/Engine/Source/Developer/SettingsEditor/Private/Widgets/SSettingsEditorCheckoutNotice.cpp
--- a/Engine/Source/Developer/SettingsEditor/Private/Widgets/SSettingsEditorCheckoutNotice.cpp
+++ b/Engine/Source/Developer/SettingsEditor/Private/Widgets/SSettingsEditorCheckoutNotice.cpp
@@ -1,9 +1,5 @@
 // Copyright 1998-2014 Epic Games, Inc. All Rights Reserved.
 
-/*=============================================================================
-	SSettingsEditorCheckoutNotice.cpp: Implements the SSettingsEditorCheckoutNotice class.
-=============================================================================*/
-
 #include "SettingsEditorPrivatePCH.h"
 
 
diff --git a/Engine/Source/Developer/SettingsEditor/Private/Widgets/SSettingsEditorCheckoutNotice.h b/Engine/Source/Developer/SettingsEditor/Private/Widgets/SSettingsEditorCheckoutNotice.h
--- a/Engine/Source/Developer/SettingsEditor/Private/Widgets/SSettingsEditorCheckoutNotice.h
+++ b/Engine/Source/Developer/SettingsEditor/Private/Widgets/SSettingsEditorCheckoutNotice.h
@@ -1,9 +1,5 @@
 // Copyright 1998-2014 Epic Games, Inc. All Rights Reserved.
 
-/*=============================================================================
-	SSettingsEditorCheckoutNotice.h: Declares the SSettingsEditorCheckoutNotice class.
-=============================================================================*/
-
 #pragma once
 
 
diff --git a/Engine/Source/Developer/SettingsEditor/Public/Interfaces/ISettingsEditorModel.h b/Engine/Source/Developer/SettingsEditor/Public/Interfaces/ISettingsEditorModel.h
--- a/Engine/Source/Developer/SettingsEditor/Public/Interfaces/ISettingsEditorModel.h
+++ b/Engine/Source/Developer/SettingsEditor/Public/Interfaces/ISettingsEditorModel.h
@@ -1,9 +1,5 @@
 // Copyright 1998-2014 Epic Games, Inc. All Rights Reserved.
 
-/*=============================================================================
-	ISettingsEditorModel.h: Declares the ISettingsEditorModel interface.
-=============================================================================*/
-
 #pragma once
 
 
diff --git a/Engine/Source/Developer/SettingsEditor/Public/Interfaces/ISettingsEditorModule.h b/Engine/Source/Developer/SettingsEditor/Public/Interfaces/ISettingsEditorModule.h
--- a/Engine/Source/Developer/SettingsEditor/Public/Interfaces/ISettingsEditorModule.h
+++ b/Engine/Source/Developer/SettingsEditor/Public/Interfaces/ISettingsEditorModule.h
@@ -1,9 +1,5 @@
 // Copyright 1998-2014 Epic Games, Inc. All Rights Reserved.
 
-/*=============================================================================
-	ISettingsEditorModule.h: Declares the ISettingsEditorModule interface.
-=============================================================================*/
-
 #pragma once
 
 
diff --git a/Engine/Source/Developer/SettingsEditor/Public/SettingsEditor.h b/Engine/Source/Developer/SettingsEditor/Public/SettingsEditor.h
--- a/Engine/Source/Developer/SettingsEditor/Public/SettingsEditor.h
+++ b/Engine/Source/Developer/SettingsEditor/Public/SettingsEditor.h
@@ -1,9 +1,5 @@
 // Copyright 1998-2014 Epic Games, Inc. All Rights Reserved.
 
-/*=============================================================================
-	SettingsEditor.h: SettingsEditor module public header file.
-=============================================================================*/
-
 #pragma once
 
 
diff --git a/Engine/Source/Developer/SlateReflector/Private/Models/WidgetReflectorNode.h b/Engine/Source/Developer/SlateReflector/Private/Models/WidgetReflectorNode.h
--- a/Engine/Source/Developer/SlateReflector/Private/Models/WidgetReflectorNode.h
+++ b/Engine/Source/Developer/SlateReflector/Private/Models/WidgetReflectorNode.h
@@ -1,9 +1,5 @@
 // Copyright 1998-2014 Epic Games, Inc. All Rights Reserved.
 
-/*=============================================================================
-	WidgetReflectorNode.h: Declares the FWidgetReflectorNode class.
-=============================================================================*/
-
 #pragma once
 
 
