fixed issue with platforms being cut off in the cook by the book menu
TTP331677
#ue4
#ufe

diff --git a/Engine/Source/Developer/SessionLauncher/Private/SessionLauncherPrivatePCH.h b/Engine/Source/Developer/SessionLauncher/Private/SessionLauncherPrivatePCH.h
--- a/Engine/Source/Developer/SessionLauncher/Private/SessionLauncherPrivatePCH.h
+++ b/Engine/Source/Developer/SessionLauncher/Private/SessionLauncherPrivatePCH.h
@@ -55,6 +55,7 @@ static TSharedRef<SButton> MakeImageButton( const FSlateBrush* ButtonImage, cons
 
 #include "SSessionLauncherMapListRow.h"
 #include "SSessionLauncherCultureListRow.h"
+#include "SSessionLauncherPlatformListRow.h"
 #include "SSessionLauncherCookedPlatforms.h"
 #include "SSessionLauncherCookByTheBookSettings.h"
 #include "SSessionLauncherCookOnTheFlySettings.h"
diff --git a/Engine/Source/Developer/SessionLauncher/Private/Widgets/Cook/SSessionLauncherCookedPlatforms.h b/Engine/Source/Developer/SessionLauncher/Private/Widgets/Cook/SSessionLauncherCookedPlatforms.h
--- a/Engine/Source/Developer/SessionLauncher/Private/Widgets/Cook/SSessionLauncherCookedPlatforms.h
+++ b/Engine/Source/Developer/SessionLauncher/Private/Widgets/Cook/SSessionLauncherCookedPlatforms.h
@@ -94,36 +94,35 @@ protected:
 	 *
 	 * @return Platform menu widget.
 	 */
-	TSharedRef<SWidget> MakePlatformMenu( ) const
+	TSharedRef<SWidget> MakePlatformMenu( )
 	{
 		TArray<ITargetPlatform*> Platforms = GetTargetPlatformManager()->GetTargetPlatforms();
 
 		if (Platforms.Num() > 0)
 		{
-			TSharedRef<SVerticalBox> VerticalBox = SNew(SVerticalBox);
-
+			PlatformList.Reset();
 			for (int32 PlatformIndex = 0; PlatformIndex < Platforms.Num(); ++PlatformIndex)
 			{
 				FString PlatformName = Platforms[PlatformIndex]->PlatformName();
 
-				VerticalBox->AddSlot()
-					.AutoHeight()
-					.Padding(0.0f, 4.0f, 0.0f, 0.0f)
-					[
-						SNew(SCheckBox)
-							.IsChecked(this, &SSessionLauncherCookedPlatforms::HandlePlatformMenuEntryIsChecked, PlatformName)
-							.OnCheckStateChanged(this, &SSessionLauncherCookedPlatforms::HandlePlatformMenuEntryCheckedStateChanged, PlatformName)
-							.Padding(FMargin(4.0f, 0.0f))
-							.Content()
-							[
-								SNew(STextBlock)
-									.ColorAndOpacity(this, &SSessionLauncherCookedPlatforms::HandlePlatformMenuEntryColorAndOpacity, PlatformName)
-									.Text(PlatformName)
-							]
-					];
+				PlatformList.Add(MakeShareable(new FString(PlatformName)));
 			}
 
-			return VerticalBox;
+			SAssignNew(PlatformListView, SListView<TSharedPtr<FString> >)
+				.HeaderRow(
+				SNew(SHeaderRow)
+				.Visibility(EVisibility::Collapsed)
+
+				+ SHeaderRow::Column("PlatformName")
+				.DefaultLabel(LOCTEXT("PlatformListPlatformNameColumnHeader", "Platform"))
+				.FillWidth(1.0f)
+				)
+				.ItemHeight(16.0f)
+				.ListItemsSource(&PlatformList)
+				.OnGenerateRow(this, &SSessionLauncherCookedPlatforms::HandlePlatformListViewGenerateRow)
+				.SelectionMode(ESelectionMode::None);
+
+			return PlatformListView.ToSharedRef();
 		}
 
 		return SNew(STextBlock)
@@ -167,24 +166,6 @@ private:
 		return EVisibility::Collapsed;
 	}
 
-	// Callback for changing the checked state of a platform menu check box. 
-	void HandlePlatformMenuEntryCheckedStateChanged( ESlateCheckBoxState::Type CheckState, FString PlatformName )
-	{
-		ILauncherProfilePtr SelectedProfile = Model->GetSelectedProfile();
-
-		if (SelectedProfile.IsValid())
-		{
-			if (CheckState == ESlateCheckBoxState::Checked)
-			{
-				SelectedProfile->AddCookedPlatform(PlatformName);
-			}
-			else
-			{
-				SelectedProfile->RemoveCookedPlatform(PlatformName);
-			}
-		}
-	}
-
 	// Callback for getting the color of a platform menu check box.
 	FSlateColor HandlePlatformMenuEntryColorAndOpacity( FString PlatformName ) const
 	{
@@ -206,26 +187,25 @@ private:
 		return FLinearColor::Yellow;
 	}
 
-	// Callback for determining whether a platform menu entry is checked.
-	ESlateCheckBoxState::Type HandlePlatformMenuEntryIsChecked( FString PlatformName ) const
+	// Handles generating a row widget in the map list view.
+	TSharedRef<ITableRow> HandlePlatformListViewGenerateRow( TSharedPtr<FString> InItem, const TSharedRef<STableViewBase>& OwnerTable )
 	{
-		ILauncherProfilePtr SelectedProfile = Model->GetSelectedProfile();
-
-		if (SelectedProfile.IsValid())
-		{
-			if (SelectedProfile->GetCookedPlatforms().Contains(PlatformName))
-			{
-				return ESlateCheckBoxState::Checked;
-			}
-		}
-
-		return ESlateCheckBoxState::Unchecked;
+		return SNew(SSessionLauncherPlatformListRow, Model.ToSharedRef())
+			.PlatformName(InItem)
+			.OwnerTableView(OwnerTable);
 	}
 
 private:
 
 	// Holds a pointer to the data model.
 	FSessionLauncherModelPtr Model;
+
+	// Holds the platform list.
+	TArray<TSharedPtr<FString> > PlatformList;
+
+	// Holds the platform list view.
+	TSharedPtr<SListView<TSharedPtr<FString> > > PlatformListView;
+
 };
 
 
diff --git a/Engine/Source/Developer/SessionLauncher/Private/Widgets/Cook/SSessionLauncherPlatformListRow.h b/Engine/Source/Developer/SessionLauncher/Private/Widgets/Cook/SSessionLauncherPlatformListRow.h
new file mode 100644
--- /dev/null
+++ b/Engine/Source/Developer/SessionLauncher/Private/Widgets/Cook/SSessionLauncherPlatformListRow.h
@@ -0,0 +1,111 @@
+// Copyright 1998-2014 Epic Games, Inc. All Rights Reserved.
+
+#pragma once
+
+
+#define LOCTEXT_NAMESPACE "SSessionLauncherPlatformListRow"
+
+
+/**
+ * Implements a row widget for map list.
+ */
+class SSessionLauncherPlatformListRow
+	: public SMultiColumnTableRow<TSharedPtr<FString> >
+{
+public:
+
+	SLATE_BEGIN_ARGS(SSessionLauncherPlatformListRow) { }
+		SLATE_ATTRIBUTE(FString, HighlightString)
+		SLATE_ARGUMENT(TSharedPtr<STableViewBase>, OwnerTableView)
+		SLATE_ARGUMENT(TSharedPtr<FString>, PlatformName)
+	SLATE_END_ARGS()
+
+public:
+
+	/**
+	 * Constructs the widget.
+	 *
+	 * @param InArgs The construction arguments.
+	 * @param InProfileManager The profile manager to use.
+	 */
+	void Construct( const FArguments& InArgs, const FSessionLauncherModelRef& InModel )
+	{
+		HighlightString = InArgs._HighlightString;
+		PlatformName = InArgs._PlatformName;
+		Model = InModel;
+
+		SMultiColumnTableRow<TSharedPtr<FString> >::Construct(FSuperRowType::FArguments(), InArgs._OwnerTableView.ToSharedRef());
+	}
+
+public:
+
+	/**
+	 * Generates the widget for the specified column.
+	 *
+	 * @param ColumnName The name of the column to generate the widget for.
+	 * @return The widget.
+	 */
+	virtual TSharedRef<SWidget> GenerateWidgetForColumn( const FName& ColumnName ) override
+	{
+		if (ColumnName == "PlatformName")
+		{
+			return SNew(SCheckBox)
+				.IsChecked(this, &SSessionLauncherPlatformListRow::HandleCheckBoxIsChecked)
+				.OnCheckStateChanged(this, &SSessionLauncherPlatformListRow::HandleCheckBoxCheckStateChanged)
+				.Padding(FMargin(6.0, 2.0))
+				[
+					SNew(STextBlock)
+						.Text(*PlatformName)			
+				];
+		}
+
+		return SNullWidget::NullWidget;
+	}
+
+private:
+
+	// Callback for changing the checked state of the check box.
+	void HandleCheckBoxCheckStateChanged( ESlateCheckBoxState::Type NewState )
+	{
+		ILauncherProfilePtr SelectedProfile = Model->GetSelectedProfile();
+
+		if (SelectedProfile.IsValid())
+		{
+			if (NewState == ESlateCheckBoxState::Checked)
+			{
+				SelectedProfile->AddCookedPlatform(*PlatformName);
+			}
+			else
+			{
+				SelectedProfile->RemoveCookedPlatform(*PlatformName);
+			}
+		}
+	}
+
+	// Callback for determining the checked state of the check box.
+	ESlateCheckBoxState::Type HandleCheckBoxIsChecked( ) const
+	{
+		ILauncherProfilePtr SelectedProfile = Model->GetSelectedProfile();
+
+		if (SelectedProfile->IsValidForLaunch() && SelectedProfile->GetCookedPlatforms().Contains(*PlatformName))
+		{
+			return ESlateCheckBoxState::Checked;
+		}
+
+		return ESlateCheckBoxState::Unchecked;
+	}
+
+private:
+
+	// Holds the highlight string for the log message.
+	TAttribute<FString> HighlightString;
+
+	// Holds the platform's name.
+	TSharedPtr<FString> PlatformName;
+
+	// Holds a pointer to the data model.
+	FSessionLauncherModelPtr Model;
+};
+
+
+#undef LOCTEXT_NAMESPACE
