fix for a list issue for the platforms in UFE
TTP331677
#ue4
#ufe

diff --git a/Engine/Source/Developer/SessionLauncher/Private/Widgets/Cook/SSessionLauncherCookByTheBookSettings.cpp b/Engine/Source/Developer/SessionLauncher/Private/Widgets/Cook/SSessionLauncherCookByTheBookSettings.cpp
--- a/Engine/Source/Developer/SessionLauncher/Private/Widgets/Cook/SSessionLauncherCookByTheBookSettings.cpp
+++ b/Engine/Source/Developer/SessionLauncher/Private/Widgets/Cook/SSessionLauncherCookByTheBookSettings.cpp
@@ -56,7 +56,7 @@ TSharedRef<SWidget> SSessionLauncherCookByTheBookSettings::MakeComplexWidget()
 
 		+ SVerticalBox::Slot()
 		.AutoHeight()
-		.MaxHeight(320.0f)
+		.MaxHeight(256.0f)
 		[
 			SNew(SBorder)
 			.BorderImage(FEditorStyle::GetBrush("ToolPanel.GroupBorder"))
@@ -74,7 +74,7 @@ TSharedRef<SWidget> SSessionLauncherCookByTheBookSettings::MakeComplexWidget()
 				]
 
 				+ SVerticalBox::Slot()
-					.AutoHeight()
+					.FillHeight(1.0)
 					.Padding(0.0f, 2.0f, 0.0f, 0.0f)
 					[
 						SNew(SSessionLauncherCookedPlatforms, Model.ToSharedRef())
diff --git a/Engine/Source/Developer/SessionLauncher/Private/Widgets/Cook/SSessionLauncherCookedPlatforms.h b/Engine/Source/Developer/SessionLauncher/Private/Widgets/Cook/SSessionLauncherCookedPlatforms.h
--- a/Engine/Source/Developer/SessionLauncher/Private/Widgets/Cook/SSessionLauncherCookedPlatforms.h
+++ b/Engine/Source/Developer/SessionLauncher/Private/Widgets/Cook/SSessionLauncherCookedPlatforms.h
@@ -29,15 +29,29 @@ public:
 	{
 		Model = InModel;
 
+		MakePlatformMenu();
+
 		ChildSlot
 		[
 			SNew(SVerticalBox)
 
 			+ SVerticalBox::Slot()
-				.AutoHeight()
+				.FillHeight(1.0f)
 				[
 					// platform menu
-					MakePlatformMenu()
+					SAssignNew(PlatformListView, SListView<TSharedPtr<FString> >)
+					.HeaderRow(
+					SNew(SHeaderRow)
+					.Visibility(EVisibility::Collapsed)
+
+					+ SHeaderRow::Column("PlatformName")
+					.DefaultLabel(LOCTEXT("PlatformListPlatformNameColumnHeader", "Platform"))
+					.FillWidth(1.0f)
+					)
+					.ItemHeight(16.0f)
+					.ListItemsSource(&PlatformList)
+					.OnGenerateRow(this, &SSessionLauncherCookedPlatforms::HandlePlatformListViewGenerateRow)
+					.SelectionMode(ESelectionMode::None)
 				]
 
 			+ SVerticalBox::Slot()
@@ -94,7 +108,7 @@ protected:
 	 *
 	 * @return Platform menu widget.
 	 */
-	TSharedRef<SWidget> MakePlatformMenu( )
+	void MakePlatformMenu( )
 	{
 		TArray<ITargetPlatform*> Platforms = GetTargetPlatformManager()->GetTargetPlatforms();
 
@@ -107,27 +121,7 @@ protected:
 
 				PlatformList.Add(MakeShareable(new FString(PlatformName)));
 			}
-
-			SAssignNew(PlatformListView, SListView<TSharedPtr<FString> >)
-				.HeaderRow(
-				SNew(SHeaderRow)
-				.Visibility(EVisibility::Collapsed)
-
-				+ SHeaderRow::Column("PlatformName")
-				.DefaultLabel(LOCTEXT("PlatformListPlatformNameColumnHeader", "Platform"))
-				.FillWidth(1.0f)
-				)
-				.ItemHeight(16.0f)
-				.ListItemsSource(&PlatformList)
-				.OnGenerateRow(this, &SSessionLauncherCookedPlatforms::HandlePlatformListViewGenerateRow)
-				.SelectionMode(ESelectionMode::None);
-
-			return PlatformListView.ToSharedRef();
 		}
-
-		return SNew(STextBlock)
-			.ColorAndOpacity(FLinearColor::Red)
-			.Text(LOCTEXT("NoPlatformsFoundErrorText", "Error: No platforms found."));
 	}
 
 private:
