Avoid crash in UDN Advanced tab if there are no profiles for the list box.

diff --git a/Engine/Source/Developer/SessionLauncher/Private/Widgets/Cook/SSessionLauncherPlatformListRow.h b/Engine/Source/Developer/SessionLauncher/Private/Widgets/Cook/SSessionLauncherPlatformListRow.h
--- a/Engine/Source/Developer/SessionLauncher/Private/Widgets/Cook/SSessionLauncherPlatformListRow.h
+++ b/Engine/Source/Developer/SessionLauncher/Private/Widgets/Cook/SSessionLauncherPlatformListRow.h
@@ -87,9 +87,12 @@ private:
 	{
 		ILauncherProfilePtr SelectedProfile = Model->GetSelectedProfile();
 
-		if (SelectedProfile->IsValidForLaunch() && SelectedProfile->GetCookedPlatforms().Contains(*PlatformName))
+		if (SelectedProfile.IsValid())
 		{
-			return ESlateCheckBoxState::Checked;
+			if (SelectedProfile->IsValidForLaunch() && SelectedProfile->GetCookedPlatforms().Contains(*PlatformName))
+			{
+				return ESlateCheckBoxState::Checked;
+			}
 		}
 
 		return ESlateCheckBoxState::Unchecked;
