fix for showing invalid platform message and not providing the information as to why
TTP332224
#ue4
#ufe

diff --git a/Engine/Source/Developer/LauncherServices/Private/Profiles/LauncherProfile.h b/Engine/Source/Developer/LauncherServices/Private/Profiles/LauncherProfile.h
--- a/Engine/Source/Developer/LauncherServices/Private/Profiles/LauncherProfile.h
+++ b/Engine/Source/Developer/LauncherServices/Private/Profiles/LauncherProfile.h
@@ -264,6 +264,11 @@ public:
 		return ValidationErrors.Contains(Error);
 	}
 
+	virtual FString GetInvalidPlatform() const override
+	{
+		return InvalidPlatform;
+	}
+
 	virtual bool IsBuilding() const override
 	{
 		return BuildGame;
@@ -781,6 +786,14 @@ protected:
 					ValidationErrors.Add(ELauncherProfileValidationErrors::NoPlatformSDKInstalled);
 					ILauncherServicesModule& LauncherServicesModule = FModuleManager::GetModuleChecked<ILauncherServicesModule>(TEXT("LauncherServices"));
 					LauncherServicesModule.BroadcastLauncherServicesSDKNotInstalled(PlatformName, NotInstalledDocLink);
+					if (!Platform)
+					{
+						CookedPlatforms.Remove(PlatformName);
+					}
+					else
+					{
+						InvalidPlatform = PlatformName;
+					}
 					return;
 				}
 			}
@@ -809,6 +822,7 @@ protected:
 							ValidationErrors.Add(ELauncherProfileValidationErrors::NoPlatformSDKInstalled);
 							ILauncherServicesModule& LauncherServicesModule = FModuleManager::GetModuleChecked<ILauncherServicesModule>(TEXT("LauncherServices"));
 							LauncherServicesModule.BroadcastLauncherServicesSDKNotInstalled(PlatformName, NotInstalledDocLink);
+							DeployedDeviceGroup->RemoveDevice(DeviceId);
 							return;
 						}
 					}
@@ -910,6 +924,7 @@ private:
 
 	// Holds the collection of validation errors.
 	TArray<ELauncherProfileValidationErrors::Type> ValidationErrors;
+	FString InvalidPlatform;
 
     // Holds the time out time for the cook on the fly server
     uint32 Timeout;
diff --git a/Engine/Source/Developer/LauncherServices/Public/Interfaces/ILauncherProfile.h b/Engine/Source/Developer/LauncherServices/Public/Interfaces/ILauncherProfile.h
--- a/Engine/Source/Developer/LauncherServices/Public/Interfaces/ILauncherProfile.h
+++ b/Engine/Source/Developer/LauncherServices/Public/Interfaces/ILauncherProfile.h
@@ -202,6 +202,8 @@ public:
 	 */
 	virtual bool HasValidationError( ELauncherProfileValidationErrors::Type Error ) const = 0;
 
+	virtual FString GetInvalidPlatform() const = 0;
+
 	/**
 	 * Checks whether devices of the specified platform can be deployed to.
 	 *
diff --git a/Engine/Source/Developer/SessionLauncher/Private/Widgets/Cook/SSessionLauncherPlatformListRow.h b/Engine/Source/Developer/SessionLauncher/Private/Widgets/Cook/SSessionLauncherPlatformListRow.h
--- a/Engine/Source/Developer/SessionLauncher/Private/Widgets/Cook/SSessionLauncherPlatformListRow.h
+++ b/Engine/Source/Developer/SessionLauncher/Private/Widgets/Cook/SSessionLauncherPlatformListRow.h
@@ -89,7 +89,8 @@ private:
 
 		if (SelectedProfile.IsValid())
 		{
-			if (SelectedProfile->IsValidForLaunch() && SelectedProfile->GetCookedPlatforms().Contains(*PlatformName))
+			SelectedProfile->IsValidForLaunch();
+			if (SelectedProfile->GetCookedPlatforms().Contains(*PlatformName))
 			{
 				return ESlateCheckBoxState::Checked;
 			}
diff --git a/Engine/Source/Developer/SessionLauncher/Private/Widgets/SSessionLauncherValidation.h b/Engine/Source/Developer/SessionLauncher/Private/Widgets/SSessionLauncherValidation.h
--- a/Engine/Source/Developer/SessionLauncher/Private/Widgets/SSessionLauncherValidation.h
+++ b/Engine/Source/Developer/SessionLauncher/Private/Widgets/SSessionLauncherValidation.h
@@ -103,7 +103,7 @@ public:
 			+ SVerticalBox::Slot()
 				.AutoHeight()
 				[
-					MakeValidationMessage(TEXT("Icons.Error"), LOCTEXT("NoPlatformSDKInstalled", "A required platform SDK is missing.").ToString(), ELauncherProfileValidationErrors::NoPlatformSDKInstalled)
+					MakeCallbackMessage(TEXT("Icons.Error"), ELauncherProfileValidationErrors::NoPlatformSDKInstalled)
 				]
 		];
 	}
@@ -139,6 +139,34 @@ protected:
 			];
 	}
 
+	/**
+	 * Creates a widget for a validation message.
+	 *
+	 * @param IconName The name of the message icon.
+	 * @param MessageType The message type.
+	 */
+	TSharedRef<SWidget> MakeCallbackMessage( const TCHAR* IconName,ELauncherProfileValidationErrors::Type Message )
+	{
+		return SNew(SHorizontalBox)
+			.Visibility(this, &SSessionLauncherValidation::HandleValidationMessageVisibility, Message)
+
+		+ SHorizontalBox::Slot()
+			.AutoWidth()
+			.Padding(2.0)
+			[
+				SNew(SImage)
+					.Image(FEditorStyle::GetBrush(IconName))
+			]
+
+		+ SHorizontalBox::Slot()
+			.AutoWidth()
+			.VAlign(VAlign_Center)
+			[
+				SNew(STextBlock)
+					.Text(this, &SSessionLauncherValidation::HandleValidationMessage, Message)
+			];
+	}
+
 private:
 
 	// Callback for getting the visibility state of a validation message.
@@ -157,6 +185,21 @@ private:
 		return EVisibility::Collapsed;
 	}
 
+	FString	HandleValidationMessage( ELauncherProfileValidationErrors::Type Error ) const
+	{
+		ILauncherProfilePtr SelectedProfile = Model->GetSelectedProfile();
+
+		if (SelectedProfile.IsValid())
+		{
+			if (SelectedProfile->HasValidationError(Error))
+			{
+				return LOCTEXT("NoPlatformSDKInstalled", "A required platform SDK is mising: ").ToString() + SelectedProfile->GetInvalidPlatform();
+			}
+		}
+
+		return TEXT("");
+	}
+
 private:
 
 	// Holds a pointer to the data model.
