fix for crash when starting up UFE
#ue4

diff --git a/Engine/Source/Developer/SessionLauncher/Private/Models/SessionLauncherCommands.h b/Engine/Source/Developer/SessionLauncher/Private/Models/SessionLauncherCommands.h
--- a/Engine/Source/Developer/SessionLauncher/Private/Models/SessionLauncherCommands.h
+++ b/Engine/Source/Developer/SessionLauncher/Private/Models/SessionLauncherCommands.h
@@ -15,8 +15,8 @@ public:
 	 */
 	FSessionLauncherCommands( )
 		: TCommands<FSessionLauncherCommands>(
-			"LauncherCommand",
-			NSLOCTEXT("Contexts", "LauncherCommand", "Launcher Command"),
+			"SessionLauncherCommand",
+			NSLOCTEXT("Contexts", "SessionLauncherCommand", "Launcher Command"),
 			NAME_None,
 			FEditorStyle::GetStyleSetName()
 	)
